Merging the removal of bSupportAR using Dev-VR_to_Release-4.20

diff --git a/Engine/Plugins/Runtime/AppleARKit/Source/AppleARKit/Private/AppleARKitSystem.cpp b/Engine/Plugins/Runtime/AppleARKit/Source/AppleARKit/Private/AppleARKitSystem.cpp
--- a/Engine/Plugins/Runtime/AppleARKit/Source/AppleARKit/Private/AppleARKitSystem.cpp
+++ b/Engine/Plugins/Runtime/AppleARKit/Source/AppleARKit/Private/AppleARKitSystem.cpp
@@ -1594,12 +1594,8 @@ namespace AppleARKitSupport
 		// Handle older iOS devices somehow calling this
 		if (FAppleARKitAvailability::SupportsARKit10())
 		{
-			const bool bIsARApp = GetDefault<UGeneralProjectSettings>()->bSupportAR;
-			if (bIsARApp)
-			{
-				auto NewARKitSystem = NewARSystem<FAppleARKitSystem>();
-				return NewARKitSystem;
-			}
+            auto NewARKitSystem = NewARSystem<FAppleARKitSystem>();
+            return NewARKitSystem;
 		}
 #endif
 		return TSharedPtr<class FAppleARKitSystem, ESPMode::ThreadSafe>();
diff --git a/Engine/Source/Runtime/AugmentedReality/Private/ARBlueprintLibrary.cpp b/Engine/Source/Runtime/AugmentedReality/Private/ARBlueprintLibrary.cpp
--- a/Engine/Source/Runtime/AugmentedReality/Private/ARBlueprintLibrary.cpp
+++ b/Engine/Source/Runtime/AugmentedReality/Private/ARBlueprintLibrary.cpp
@@ -75,7 +75,7 @@ EARTrackingQuality UARBlueprintLibrary::GetTrackingQuality()
 
 void UARBlueprintLibrary::StartARSession(UARSessionConfig* SessionConfig)
 {
-	static const TCHAR NotARApp_Warning[] = TEXT("Attempting to Start an AR session but there is no AR system. To use AR, enable bIsARApp under Project Settings.");
+	static const TCHAR NotARApp_Warning[] = TEXT("Attempting to start an AR session but there is no AR plugin configured. To use AR, enable the proper AR plugin in the Plugin Settings.");
 	
 	auto ARSystem = GetARSystem();
 	if (ensureAlwaysMsgf(ARSystem.IsValid(), NotARApp_Warning))
diff --git a/Engine/Source/Runtime/AugmentedReality/Private/ARSystem.cpp b/Engine/Source/Runtime/AugmentedReality/Private/ARSystem.cpp
--- a/Engine/Source/Runtime/AugmentedReality/Private/ARSystem.cpp
+++ b/Engine/Source/Runtime/AugmentedReality/Private/ARSystem.cpp
@@ -44,25 +44,11 @@ EARTrackingQuality FARSystemBase::GetTrackingQuality() const
 
 void FARSystemBase::StartARSession(UARSessionConfig* InSessionConfig)
 {
-	static const TCHAR NotARApp_Warning[] = TEXT("To use AR, enable Supports AR under Project Settings.");
-	
-	const bool bIsARApp = GetDefault<UGeneralProjectSettings>()->bSupportAR;
-	if (ensureAlwaysMsgf(bIsARApp, NotARApp_Warning))
-	{
-		if (GetARSessionStatus().Status != EARSessionStatus::Running)
-		{
-			OnStartARSession(InSessionConfig);
-			ARSettings = InSessionConfig;
-		}
-	}
-	else
-	{
-#if !(UE_BUILD_SHIPPING || UE_BUILD_TEST)
-		// Ensures don't show up on iOS, but we definitely want a developer to see this
-		// Their AR project just doesn't make sense unless they have enabled the AR setting.
-		GEngine->AddOnScreenDebugMessage((uint64)((PTRINT)this), 3600.0f, FColor(255,48,16),NotARApp_Warning);
-#endif //!(UE_BUILD_SHIPPING || UE_BUILD_TEST)
-	}
+    if (GetARSessionStatus().Status != EARSessionStatus::Running)
+    {
+        OnStartARSession(InSessionConfig);
+        ARSettings = InSessionConfig;
+    }
 }
 
 void FARSystemBase::PauseARSession()
diff --git a/Engine/Source/Runtime/EngineSettings/Classes/GeneralProjectSettings.h b/Engine/Source/Runtime/EngineSettings/Classes/GeneralProjectSettings.h
--- a/Engine/Source/Runtime/EngineSettings/Classes/GeneralProjectSettings.h
+++ b/Engine/Source/Runtime/EngineSettings/Classes/GeneralProjectSettings.h
@@ -83,9 +83,9 @@ class ENGINESETTINGS_API UGeneralProjectSettings
 	UPROPERTY()
 	bool bStartInAR_DEPRECATED;
 
-	/** Will this app support Augmented Reality (AR) in some capacity? */
-	UPROPERTY(config, EditAnywhere, Category = Settings)
-	bool bSupportAR;
+    /** No longer used; AR support is determined by included plugins */
+	UPROPERTY()
+	bool bSupportAR_DEPRECATED;
 	
 	UPROPERTY(config, EditAnywhere, Category = Settings)
 	bool bAllowWindowResize;
diff --git a/Engine/Source/Runtime/EngineSettings/Private/EngineSettingsModule.cpp b/Engine/Source/Runtime/EngineSettings/Private/EngineSettingsModule.cpp
--- a/Engine/Source/Runtime/EngineSettings/Private/EngineSettingsModule.cpp
+++ b/Engine/Source/Runtime/EngineSettings/Private/EngineSettingsModule.cpp
@@ -72,7 +72,7 @@ UGeneralProjectSettings::UGeneralProjectSettings( const FObjectInitializer& Obje
 	, bUseBorderlessWindow(false)
 	, bStartInVR(false)
 	, bStartInAR_DEPRECATED(false)
-	, bSupportAR(false)
+	, bSupportAR_DEPRECATED(false)
 	, bAllowWindowResize(true)
 	, bAllowClose(true)
 	, bAllowMaximize(true)
