Merging fix for UE-61067 using Dev-VR_to_Release-4.20

diff --git a/Engine/Source/Runtime/AugmentedReality/Private/ARSkyLight.cpp b/Engine/Source/Runtime/AugmentedReality/Private/ARSkyLight.cpp
--- a/Engine/Source/Runtime/AugmentedReality/Private/ARSkyLight.cpp
+++ b/Engine/Source/Runtime/AugmentedReality/Private/ARSkyLight.cpp
@@ -20,6 +20,7 @@ AARSkyLight::AARSkyLight(const FObjectInitializer& ObjectInitializer)
 
 	// We only work with environment probes that generate a cube map
 	GetLightComponent()->SourceType = ESkyLightSourceType::SLS_SpecifiedCubemap;
+	GetLightComponent()->bLowerHemisphereIsBlack = false;
 }
 
 // @todo - JoeG add double buffering of the cube map textures so we can blend the updates rather than popping
