Merging UE-61204 fix using Dev-VR_to_Release-4.20

diff --git a/Engine/Source/Runtime/AugmentedReality/Public/ARTypes.h b/Engine/Source/Runtime/AugmentedReality/Public/ARTypes.h
--- a/Engine/Source/Runtime/AugmentedReality/Public/ARTypes.h
+++ b/Engine/Source/Runtime/AugmentedReality/Public/ARTypes.h
@@ -171,14 +171,23 @@ class AUGMENTEDREALITY_API UARCandidateImage :
 
 public:
 	/** @see CandidateTexture */
+	UFUNCTION(BlueprintPure, Category = "AR AugmentedReality|Image Detection")
 	UTexture2D* GetCandidateTexture() const { return CandidateTexture; }
+
 	/** @see FriendlyName */
+	UFUNCTION(BlueprintPure, Category = "AR AugmentedReality|Image Detection")
 	const FString& GetFriendlyName() const { return FriendlyName; }
+
 	/** @see Width */
+	UFUNCTION(BlueprintPure, Category = "AR AugmentedReality|Image Detection")
 	float GetPhysicalWidth() const { return Width; }
+
 	/** @see Height */
+	UFUNCTION(BlueprintPure, Category = "AR AugmentedReality|Image Detection")
 	float GetPhysicalHeight() const { return Height; }
+
 	/** @see Orientation */
+	UFUNCTION(BlueprintPure, Category = "AR AugmentedReality|Image Detection")
 	EARCandidateImageOrientation GetOrientation() const { return Orientation; }
 
 private:
@@ -214,12 +223,21 @@ class AUGMENTEDREALITY_API UARCandidateObject :
 	
 public:
 	/** @see CandidateObjectData */
+	UFUNCTION(BlueprintPure, Category = "AR AugmentedReality|Object Detection")
 	const TArray<uint8>& GetCandidateObjectData() const { return CandidateObjectData; }
+
+	UFUNCTION(BlueprintCallable, Category = "AR AugmentedReality|Object Detection")
 	void SetCandidateObjectData(const TArray<uint8>& InCandidateObject) { CandidateObjectData = InCandidateObject; }
+
 	/** @see FriendlyName */
+	UFUNCTION(BlueprintPure, Category = "AR AugmentedReality|Object Detection")
 	const FString& GetFriendlyName() const { return FriendlyName; }
+
 	/** @see BoundingBox */
+	UFUNCTION(BlueprintPure, Category = "AR AugmentedReality|Object Detection")
 	const FBox& GetBoundingBox() const { return BoundingBox; }
+
+	UFUNCTION(BlueprintCallable, Category = "AR AugmentedReality|Object Detection")
 	void SetBoundingBox(FBox& InBoundingBox) { BoundingBox = InBoundingBox; }
 
 private:
