Added a missing const that was marking a set function as in/out

diff --git a/Engine/Source/Runtime/AugmentedReality/Public/ARTypes.h b/Engine/Source/Runtime/AugmentedReality/Public/ARTypes.h
--- a/Engine/Source/Runtime/AugmentedReality/Public/ARTypes.h
+++ b/Engine/Source/Runtime/AugmentedReality/Public/ARTypes.h
@@ -238,18 +238,18 @@ public:
 	const FBox& GetBoundingBox() const { return BoundingBox; }
 
 	UFUNCTION(BlueprintCallable, Category = "AR AugmentedReality|Object Detection")
-	void SetBoundingBox(FBox& InBoundingBox) { BoundingBox = InBoundingBox; }
+	void SetBoundingBox(const FBox& InBoundingBox) { BoundingBox = InBoundingBox; }
 
 private:
 	/** The object to detect in scenes */
 	UPROPERTY(EditAnywhere, Category = "AR Candidate Object")
 	TArray<uint8> CandidateObjectData;
 	
-	/** The friendly name to report back when the image is detected in scenes */
+	/** The friendly name to report back when the object is detected in scenes */
 	UPROPERTY(EditAnywhere, Category = "AR Candidate Object")
 	FString FriendlyName;
 	
-	/** The physical width in centimeters of the object that this candidate image represents */
+	/** The physical bounds in centimeters of the object that this candidate object represents */
 	UPROPERTY(EditAnywhere, Category = "AR Candidate Image")
 	FBox BoundingBox;
 };
