Bringing over default ctor init from devvr to fix warnings.

diff --git a/Engine/Source/Runtime/AugmentedReality/Public/ARTypes.h b/Engine/Source/Runtime/AugmentedReality/Public/ARTypes.h
--- a/Engine/Source/Runtime/AugmentedReality/Public/ARTypes.h
+++ b/Engine/Source/Runtime/AugmentedReality/Public/ARTypes.h
@@ -326,7 +326,13 @@ struct AUGMENTEDREALITY_API FARVideoFormat
 	GENERATED_BODY()
 
 public:
-	FARVideoFormat() { }
+	FARVideoFormat()
+		: FPS(0)
+		, Width(0)
+		, Height(0)
+	{
+		
+	}
 	
 	FARVideoFormat(int32 InFPS, int32 InWidth, int32 InHeight)
 		: FPS(InFPS)
