Fixing FARSupportInterface reentrant crash on shutdown.

diff --git a/Engine/Source/Runtime/AugmentedReality/Private/ARBlueprintLibrary.cpp b/Engine/Source/Runtime/AugmentedReality/Private/ARBlueprintLibrary.cpp
--- a/Engine/Source/Runtime/AugmentedReality/Private/ARBlueprintLibrary.cpp
+++ b/Engine/Source/Runtime/AugmentedReality/Private/ARBlueprintLibrary.cpp
@@ -8,14 +8,14 @@
 #include "ARPin.h"
 #include "ARTrackable.h"
 
-TSharedPtr<FARSupportInterface , ESPMode::ThreadSafe> UARBlueprintLibrary::RegisteredARSystem = nullptr;
+TWeakPtr<FARSupportInterface , ESPMode::ThreadSafe> UARBlueprintLibrary::RegisteredARSystem = nullptr;
 
 EARTrackingQuality UARBlueprintLibrary::GetTrackingQuality()
 {
 	auto ARSystem = GetARSystem();
 	if (ARSystem.IsValid())
 	{
-		return ARSystem->GetTrackingQuality();
+		return ARSystem.Pin()->GetTrackingQuality();
 	}
 	else
 	{
@@ -60,7 +60,7 @@ void UARBlueprintLibrary::StartARSession(UARSessionConfig* SessionConfig)
 	auto ARSystem = GetARSystem();
 	if (ARSystem.IsValid())
 	{
-		ARSystem->StartARSession(SessionConfig);
+		ARSystem.Pin()->StartARSession(SessionConfig);
 	}
 	else
 	{
@@ -75,7 +75,7 @@ void UARBlueprintLibrary::PauseARSession()
 	auto ARSystem = GetARSystem();
 	if (ARSystem.IsValid())
 	{
-		ARSystem->PauseARSession();
+		ARSystem.Pin()->PauseARSession();
 	}
 }
 
@@ -84,7 +84,7 @@ void UARBlueprintLibrary::StopARSession()
 	auto ARSystem = GetARSystem();
 	if (ARSystem.IsValid())
 	{
-		ARSystem->StopARSession();
+		ARSystem.Pin()->StopARSession();
 	}
 }
 
@@ -93,7 +93,7 @@ FARSessionStatus UARBlueprintLibrary::GetARSessionStatus()
 	auto ARSystem = GetARSystem();
 	if (ARSystem.IsValid())
 	{
-		return ARSystem->GetARSessionStatus();
+		return ARSystem.Pin()->GetARSessionStatus();
 	}
 	else
 	{
@@ -105,7 +105,7 @@ UARSessionConfig* UARBlueprintLibrary::GetSessionConfig()
 	auto ARSystem = GetARSystem();
 	if (ARSystem.IsValid())
 	{
-		return &ARSystem->AccessSessionConfig();
+		return &ARSystem.Pin()->AccessSessionConfig();
 	}
 	else
 	{
@@ -119,7 +119,7 @@ void UARBlueprintLibrary::SetAlignmentTransform( const FTransform& InAlignmentTr
 	auto ARSystem = GetARSystem();
 	if (ARSystem.IsValid())
 	{
-		return ARSystem->SetAlignmentTransform( InAlignmentTransform );
+		return ARSystem.Pin()->SetAlignmentTransform( InAlignmentTransform );
 	}
 }
 
@@ -137,7 +137,7 @@ TArray<FARTraceResult> UARBlueprintLibrary::LineTraceTrackedObjects( const FVect
 			(bTestPlaneExtents ? EARLineTraceChannels::PlaneUsingExtent : EARLineTraceChannels::None ) |
 			(bTestPlaneBoundaryPolygon ? EARLineTraceChannels::PlaneUsingBoundaryPolygon : EARLineTraceChannels::None);
 		
-		Result = ARSystem->LineTraceTrackedObjects(ScreenCoord, ActiveTraceChannels);
+		Result = ARSystem.Pin()->LineTraceTrackedObjects(ScreenCoord, ActiveTraceChannels);
 	}
 	
 	return Result;
@@ -156,7 +156,7 @@ TArray<FARTraceResult> UARBlueprintLibrary::LineTraceTrackedObjects3D(const FVec
 			(bTestPlaneExtents ? EARLineTraceChannels::PlaneUsingExtent : EARLineTraceChannels::None) |
 			(bTestPlaneBoundaryPolygon ? EARLineTraceChannels::PlaneUsingBoundaryPolygon : EARLineTraceChannels::None);
 
-		Result = ARSystem->LineTraceTrackedObjects(Start, End, ActiveTraceChannels);
+		Result = ARSystem.Pin()->LineTraceTrackedObjects(Start, End, ActiveTraceChannels);
 	}
 
 	return Result;
@@ -169,7 +169,7 @@ TArray<UARTrackedGeometry*> UARBlueprintLibrary::GetAllGeometries()
 	auto ARSystem = GetARSystem();
 	if (ARSystem.IsValid())
 	{
-		Geometries = ARSystem->GetAllTrackedGeometries();
+		Geometries = ARSystem.Pin()->GetAllTrackedGeometries();
 	}
 	return Geometries;
 }
@@ -181,7 +181,7 @@ TArray<UARPin*> UARBlueprintLibrary::GetAllPins()
 	auto ARSystem = GetARSystem();
 	if (ARSystem.IsValid())
 	{
-		Pins = ARSystem->GetAllPins();
+		Pins = ARSystem.Pin()->GetAllPins();
 	}
 	return Pins;
 }
@@ -191,7 +191,7 @@ bool UARBlueprintLibrary::IsSessionTypeSupported(EARSessionType SessionType)
 	auto ARSystem = GetARSystem();
 	if (ARSystem.IsValid())
 	{
-		return ARSystem->IsSessionTypeSupported(SessionType);
+		return ARSystem.Pin()->IsSessionTypeSupported(SessionType);
 	}
 	return false;
 }
@@ -221,7 +221,7 @@ UARLightEstimate* UARBlueprintLibrary::GetCurrentLightEstimate()
 	auto ARSystem = GetARSystem();
 	if (ARSystem.IsValid())
 	{
-		return ARSystem->GetCurrentLightEstimate();
+		return ARSystem.Pin()->GetCurrentLightEstimate();
 	}
 	return nullptr;
 }
@@ -231,7 +231,7 @@ UARPin* UARBlueprintLibrary::PinComponent( USceneComponent* ComponentToPin, cons
 	auto ARSystem = GetARSystem();
 	if (ARSystem.IsValid())
 	{
-		return ARSystem->PinComponent( ComponentToPin, PinToWorldTransform, TrackedGeometry, DebugName );
+		return ARSystem.Pin()->PinComponent( ComponentToPin, PinToWorldTransform, TrackedGeometry, DebugName );
 	}
 	return nullptr;
 }
@@ -241,7 +241,7 @@ UARPin* UARBlueprintLibrary::PinComponentToTraceResult( USceneComponent* Compone
 	auto ARSystem = GetARSystem();
 	if (ARSystem.IsValid())
 	{
-		return ARSystem->PinComponent( ComponentToPin, TraceResult, DebugName );
+		return ARSystem.Pin()->PinComponent( ComponentToPin, TraceResult, DebugName );
 	}
 	return nullptr;
 }
@@ -251,13 +251,14 @@ void UARBlueprintLibrary::UnpinComponent( USceneComponent* ComponentToUnpin )
 	auto ARSystem = GetARSystem();
 	if (ARSystem.IsValid())
 	{
-		TArray<UARPin*> AllPins = ARSystem->GetAllPins();
+		auto PinnedARSystem = ARSystem.Pin();
+		TArray<UARPin*> AllPins = PinnedARSystem->GetAllPins();
 		const int32 AllPinsCount = AllPins.Num();
 		for (int32 i=0; i<AllPinsCount; ++i)
 		{
 			if (AllPins[i]->GetPinnedComponent() == ComponentToUnpin)
 			{
-				ARSystem->RemovePin( AllPins[i] );
+				PinnedARSystem->RemovePin( AllPins[i] );
 				return;
 			}
 		}
@@ -269,19 +270,19 @@ void UARBlueprintLibrary::RemovePin( UARPin* PinToRemove )
 	auto ARSystem = GetARSystem();
 	if (ARSystem.IsValid())
 	{
-		return ARSystem->RemovePin( PinToRemove );
+		return ARSystem.Pin()->RemovePin( PinToRemove );
 	}
 }
 
 
 
-void UARBlueprintLibrary::RegisterAsARSystem(const TSharedPtr<FARSupportInterface , ESPMode::ThreadSafe>& NewARSystem)
+void UARBlueprintLibrary::RegisterAsARSystem(const TSharedRef<FARSupportInterface , ESPMode::ThreadSafe>& NewARSystem)
 {
 	RegisteredARSystem = NewARSystem;
 }
 
 
-const TSharedPtr<FARSupportInterface , ESPMode::ThreadSafe>& UARBlueprintLibrary::GetARSystem()
+const TWeakPtr<FARSupportInterface , ESPMode::ThreadSafe>& UARBlueprintLibrary::GetARSystem()
 {
 	return RegisteredARSystem;
 }
@@ -320,7 +321,7 @@ UARTextureCameraImage* UARBlueprintLibrary::GetCameraImage()
 	auto ARSystem = GetARSystem();
 	if (ARSystem.IsValid())
 	{
-		Image = ARSystem->GetCameraImage();
+		Image = ARSystem.Pin()->GetCameraImage();
 	}
 	return Image;
 }
@@ -332,7 +333,7 @@ UARTextureCameraDepth* UARBlueprintLibrary::GetCameraDepth()
 	auto ARSystem = GetARSystem();
 	if (ARSystem.IsValid())
 	{
-		Depth = ARSystem->GetCameraDepth();
+		Depth = ARSystem.Pin()->GetCameraDepth();
 	}
 	return Depth;
 }
@@ -344,7 +345,8 @@ TArray<UARPlaneGeometry*> UARBlueprintLibrary::GetAllTrackedPlanes()
 	auto ARSystem = GetARSystem();
 	if (ARSystem.IsValid())
 	{
-		for (UARTrackedGeometry* Geo : ARSystem->GetAllTrackedGeometries())
+		auto PinnedARSystem = ARSystem.Pin();
+		for (UARTrackedGeometry* Geo : PinnedARSystem->GetAllTrackedGeometries())
 		{
 			if (UARPlaneGeometry* Item = Cast<UARPlaneGeometry>(Geo))
 			{
@@ -362,7 +364,8 @@ TArray<UARTrackedPoint*> UARBlueprintLibrary::GetAllTrackedPoints()
 	auto ARSystem = GetARSystem();
 	if (ARSystem.IsValid())
 	{
-		for (UARTrackedGeometry* Geo : ARSystem->GetAllTrackedGeometries())
+		auto PinnedARSystem = ARSystem.Pin();
+		for (UARTrackedGeometry* Geo : PinnedARSystem->GetAllTrackedGeometries())
 		{
 			if (UARTrackedPoint* Item = Cast<UARTrackedPoint>(Geo))
 			{
@@ -380,7 +383,8 @@ TArray<UARTrackedImage*> UARBlueprintLibrary::GetAllTrackedImages()
 	auto ARSystem = GetARSystem();
 	if (ARSystem.IsValid())
 	{
-		for (UARTrackedGeometry* Geo : ARSystem->GetAllTrackedGeometries())
+		auto PinnedARSystem = ARSystem.Pin();
+		for (UARTrackedGeometry* Geo : PinnedARSystem->GetAllTrackedGeometries())
 		{
 			if (UARTrackedImage* Item = Cast<UARTrackedImage>(Geo))
 			{
@@ -398,7 +402,8 @@ TArray<UAREnvironmentCaptureProbe*> UARBlueprintLibrary::GetAllTrackedEnvironmen
 	auto ARSystem = GetARSystem();
 	if (ARSystem.IsValid())
 	{
-		for (UARTrackedGeometry* Geo : ARSystem->GetAllTrackedGeometries())
+		auto PinnedARSystem = ARSystem.Pin();
+		for (UARTrackedGeometry* Geo : PinnedARSystem->GetAllTrackedGeometries())
 		{
 			if (UAREnvironmentCaptureProbe* Item = Cast<UAREnvironmentCaptureProbe>(Geo))
 			{
@@ -414,7 +419,7 @@ bool UARBlueprintLibrary::AddManualEnvironmentCaptureProbe(FVector Location, FVe
 	auto ARSystem = GetARSystem();
 	if (ARSystem.IsValid())
 	{
-		return ARSystem->AddManualEnvironmentCaptureProbe(Location, Extent);
+		return ARSystem.Pin()->AddManualEnvironmentCaptureProbe(Location, Extent);
 	}
 	return false;
 }
@@ -424,7 +429,7 @@ EARWorldMappingState UARBlueprintLibrary::GetWorldMappingStatus()
 	auto ARSystem = GetARSystem();
 	if (ARSystem.IsValid())
 	{
-		return ARSystem->GetWorldMappingStatus();
+		return ARSystem.Pin()->GetWorldMappingStatus();
 	}
 	return EARWorldMappingState::NotAvailable;
 }
@@ -434,7 +439,7 @@ TSharedPtr<FARSaveWorldAsyncTask, ESPMode::ThreadSafe> UARBlueprintLibrary::Save
 	auto ARSystem = GetARSystem();
 	if (ARSystem.IsValid())
 	{
-		return ARSystem->SaveWorld();
+		return ARSystem.Pin()->SaveWorld();
 		
 	}
 	return TSharedPtr<FARSaveWorldAsyncTask, ESPMode::ThreadSafe>();
@@ -445,7 +450,7 @@ TSharedPtr<FARGetCandidateObjectAsyncTask, ESPMode::ThreadSafe> UARBlueprintLibr
 	auto ARSystem = GetARSystem();
 	if (ARSystem.IsValid())
 	{
-		return ARSystem->GetCandidateObject(Location, Extent);
+		return ARSystem.Pin()->GetCandidateObject(Location, Extent);
 		
 	}
 	return TSharedPtr<FARGetCandidateObjectAsyncTask, ESPMode::ThreadSafe>();
@@ -456,7 +461,7 @@ TArray<FVector> UARBlueprintLibrary::GetPointCloud()
 	auto ARSystem = GetARSystem();
 	if (ARSystem.IsValid())
 	{
-		return ARSystem->GetPointCloud();
+		return ARSystem.Pin()->GetPointCloud();
 	}
 	return TArray<FVector>();
 }
@@ -466,7 +471,7 @@ TArray<FARVideoFormat> UARBlueprintLibrary::GetSupportedVideoFormats(EARSessionT
 	auto ARSystem = GetARSystem();
 	if (ARSystem.IsValid())
 	{
-		return ARSystem->GetSupportedVideoFormats(SessionType);
+		return ARSystem.Pin()->GetSupportedVideoFormats(SessionType);
 	}
 	return TArray<FARVideoFormat>();
 }
@@ -476,7 +481,7 @@ UARCandidateImage* UARBlueprintLibrary::AddRuntimeCandidateImage(UARSessionConfi
 	auto ARSystem = GetARSystem();
 	if (ensure(ARSystem.IsValid()))
 	{
-		return ARSystem->AddRuntimeCandidateImage(SessionConfig, CandidateTexture, FriendlyName, PhysicalWidth);
+		return ARSystem.Pin()->AddRuntimeCandidateImage(SessionConfig, CandidateTexture, FriendlyName, PhysicalWidth);
 	}
 	else
 	{
diff --git a/Engine/Source/Runtime/AugmentedReality/Private/ARBlueprintProxy.cpp b/Engine/Source/Runtime/AugmentedReality/Private/ARBlueprintProxy.cpp
--- a/Engine/Source/Runtime/AugmentedReality/Private/ARBlueprintProxy.cpp
+++ b/Engine/Source/Runtime/AugmentedReality/Private/ARBlueprintProxy.cpp
@@ -3,7 +3,7 @@
 #include "ARBlueprintProxy.h"
 #include "ARSystem.h"
 
-TSharedPtr<FARSupportInterface , ESPMode::ThreadSafe> UARBaseAsyncTaskBlueprintProxy::RegisteredARSystem = nullptr;
+TWeakPtr<FARSupportInterface , ESPMode::ThreadSafe> UARBaseAsyncTaskBlueprintProxy::RegisteredARSystem = nullptr;
 
 UARBaseAsyncTaskBlueprintProxy::UARBaseAsyncTaskBlueprintProxy(const FObjectInitializer& ObjectInitializer)
 	: Super(ObjectInitializer)
@@ -35,13 +35,13 @@ void UARBaseAsyncTaskBlueprintProxy::Tick(float DeltaTime)
 	}
 }
 
-void UARBaseAsyncTaskBlueprintProxy::RegisterAsARSystem(const TSharedPtr<FARSupportInterface , ESPMode::ThreadSafe>& NewARSystem)
+void UARBaseAsyncTaskBlueprintProxy::RegisterAsARSystem(const TSharedRef<FARSupportInterface , ESPMode::ThreadSafe>& NewARSystem)
 {
 	RegisteredARSystem = NewARSystem;
 }
 
 
-const TSharedPtr<FARSupportInterface , ESPMode::ThreadSafe>& UARBaseAsyncTaskBlueprintProxy::GetARSystem()
+const TWeakPtr<FARSupportInterface , ESPMode::ThreadSafe>& UARBaseAsyncTaskBlueprintProxy::GetARSystem()
 {
 	return RegisteredARSystem;
 }
@@ -58,7 +58,7 @@ void UARSaveWorldAsyncTaskBlueprintProxy::Activate()
 	auto ARSystem = GetARSystem();
 	if (ARSystem.IsValid())
 	{
-		SaveWorldTask = ARSystem->SaveWorld();
+		SaveWorldTask = ARSystem.Pin()->SaveWorld();
 		AsyncTask = SaveWorldTask;
 	}
 	else
@@ -92,7 +92,7 @@ void UARGetCandidateObjectAsyncTaskBlueprintProxy::Activate()
 	auto ARSystem = GetARSystem();
 	if (ARSystem.IsValid())
 	{
-		CandidateObjectTask = ARSystem->GetCandidateObject(Location, Extent);
+		CandidateObjectTask = ARSystem.Pin()->GetCandidateObject(Location, Extent);
 		AsyncTask = CandidateObjectTask;
 	}
 	else
diff --git a/Engine/Source/Runtime/AugmentedReality/Private/ARSupportInterface.cpp b/Engine/Source/Runtime/AugmentedReality/Private/ARSupportInterface.cpp
--- a/Engine/Source/Runtime/AugmentedReality/Private/ARSupportInterface.cpp
+++ b/Engine/Source/Runtime/AugmentedReality/Private/ARSupportInterface.cpp
@@ -17,16 +17,11 @@ FARSupportInterface ::FARSupportInterface (IARSystemSupport* InARImplementation,
 {
 }
 
-
 FARSupportInterface ::~FARSupportInterface ()
 {
 	IModularFeatures::Get().UnregisterModularFeature(FARSupportInterface ::GetModularFeatureName(), this);
-
-	UARBlueprintLibrary::RegisterAsARSystem(nullptr);
-	UARBaseAsyncTaskBlueprintProxy::RegisterAsARSystem(nullptr);
 }
 
-
 void FARSupportInterface ::InitializeARSystem()
 {
 	// Register our ability to support Unreal AR API.
@@ -34,8 +29,8 @@ void FARSupportInterface ::InitializeARSystem()
 
 	if (ARImplemention)
 	{
-		UARBlueprintLibrary::RegisterAsARSystem(MakeShareable(this));
-		UARBaseAsyncTaskBlueprintProxy::RegisterAsARSystem(MakeShareable(this));
+		UARBlueprintLibrary::RegisterAsARSystem(AsShared());
+		UARBaseAsyncTaskBlueprintProxy::RegisterAsARSystem(AsShared());
 
 		ARImplemention->OnARSystemInitialized();
 	}
diff --git a/Engine/Source/Runtime/AugmentedReality/Public/ARBlueprintLibrary.h b/Engine/Source/Runtime/AugmentedReality/Public/ARBlueprintLibrary.h
--- a/Engine/Source/Runtime/AugmentedReality/Public/ARBlueprintLibrary.h
+++ b/Engine/Source/Runtime/AugmentedReality/Public/ARBlueprintLibrary.h
@@ -231,11 +231,11 @@ public:
 	static UARCandidateImage* AddRuntimeCandidateImage(UARSessionConfig* SessionConfig, UTexture2D* CandidateTexture, FString FriendlyName, float PhysicalWidth);
 
 public:
-	static void RegisterAsARSystem(const TSharedPtr<FARSupportInterface , ESPMode::ThreadSafe>& NewArSystem);
+	static void RegisterAsARSystem(const TSharedRef<FARSupportInterface , ESPMode::ThreadSafe>& NewArSystem);
 	
 private:
-	static const TSharedPtr<FARSupportInterface , ESPMode::ThreadSafe>& GetARSystem();
-	static TSharedPtr<FARSupportInterface , ESPMode::ThreadSafe> RegisteredARSystem;
+	static const TWeakPtr<FARSupportInterface , ESPMode::ThreadSafe>& GetARSystem();
+	static TWeakPtr<FARSupportInterface , ESPMode::ThreadSafe> RegisteredARSystem;
 };
 
 
diff --git a/Engine/Source/Runtime/AugmentedReality/Public/ARBlueprintProxy.h b/Engine/Source/Runtime/AugmentedReality/Public/ARBlueprintProxy.h
--- a/Engine/Source/Runtime/AugmentedReality/Public/ARBlueprintProxy.h
+++ b/Engine/Source/Runtime/AugmentedReality/Public/ARBlueprintProxy.h
@@ -31,10 +31,10 @@ public:
 	virtual void ReportSuccess() { check(0); }
 	virtual void ReportFailure() { check(0); }
 
-	static void RegisterAsARSystem(const TSharedPtr<FARSupportInterface , ESPMode::ThreadSafe>& NewArSystem);
+	static void RegisterAsARSystem(const TSharedRef<FARSupportInterface , ESPMode::ThreadSafe>& NewArSystem);
 
 protected:
-	static const TSharedPtr<FARSupportInterface , ESPMode::ThreadSafe>& GetARSystem();
+	static const TWeakPtr<FARSupportInterface , ESPMode::ThreadSafe>& GetARSystem();
 	/** The async task to check during Tick() */
 	TSharedPtr<FARAsyncTask, ESPMode::ThreadSafe> AsyncTask;
 	
@@ -42,7 +42,7 @@ private:
 	/** True until the async task completes, then false */
 	bool bShouldTick;
 
-	static TSharedPtr<FARSupportInterface , ESPMode::ThreadSafe> RegisteredARSystem;
+	static TWeakPtr<FARSupportInterface , ESPMode::ThreadSafe> RegisteredARSystem;
 };
 
 DECLARE_DYNAMIC_MULTICAST_DELEGATE_OneParam(FARSaveWorldPin, const TArray<uint8>&, SavedWorld);
diff --git a/Engine/Source/Runtime/AugmentedReality/Public/ARSupportInterface.h b/Engine/Source/Runtime/AugmentedReality/Public/ARSupportInterface.h
--- a/Engine/Source/Runtime/AugmentedReality/Public/ARSupportInterface.h
+++ b/Engine/Source/Runtime/AugmentedReality/Public/ARSupportInterface.h
@@ -16,7 +16,7 @@ DECLARE_MULTICAST_DELEGATE_OneParam(FARSystemOnAlignmentTransformUpdated, const
  * Composition Components for tracking system features
 */
 
-class AUGMENTEDREALITY_API FARSupportInterface  : public FGCObject, public IModularFeature
+class AUGMENTEDREALITY_API FARSupportInterface : public TSharedFromThis<FARSupportInterface>, public FGCObject, public IModularFeature
 {
 public:
 	FARSupportInterface (IARSystemSupport* InARImplementation, IXRTrackingSystem* InXRTrackingSystem);
