Fix build

diff --git a/Engine/Source/Runtime/AugmentedReality/Public/ARSupportInterface.h b/Engine/Source/Runtime/AugmentedReality/Public/ARSupportInterface.h
--- a/Engine/Source/Runtime/AugmentedReality/Public/ARSupportInterface.h
+++ b/Engine/Source/Runtime/AugmentedReality/Public/ARSupportInterface.h
@@ -16,7 +16,7 @@ DECLARE_MULTICAST_DELEGATE_OneParam(FARSystemOnAlignmentTransformUpdated, const
  * Composition Components for tracking system features
 */
 
-class AUGMENTEDREALITY_API FARSupportInterface : public TSharedFromThis<FARSupportInterface>, public FGCObject, public IModularFeature
+class AUGMENTEDREALITY_API FARSupportInterface : public TSharedFromThis<FARSupportInterface, ESPMode::ThreadSafe>, public FGCObject, public IModularFeature
 {
 public:
 	FARSupportInterface (IARSystemSupport* InARImplementation, IXRTrackingSystem* InXRTrackingSystem);
