UARPlaneGeometry::DebugDraw() - PersistForSeconds and OutlineThickness were ignored for BoundaryPolygons, pass them through.

diff --git a/Engine/Source/Runtime/AugmentedReality/Private/ARTrackable.cpp b/Engine/Source/Runtime/AugmentedReality/Private/ARTrackable.cpp
--- a/Engine/Source/Runtime/AugmentedReality/Private/ARTrackable.cpp
+++ b/Engine/Source/Runtime/AugmentedReality/Private/ARTrackable.cpp
@@ -170,10 +170,10 @@ void UARPlaneGeometry::DebugDraw( UWorld* World, const FLinearColor& OutlineColo
 		for (int32 i=1; i<BoundaryPolygon.Num(); ++i)
 		{
 			const FVector NewVert = LocalToWorldTransform.TransformPosition(BoundaryPolygon[i]);
-			DrawDebugLine(World, LastVert, NewVert, OutlineRGB);
+			DrawDebugLine(World, LastVert, NewVert, OutlineRGB, PersistForSeconds > 0 ? true : false, PersistForSeconds, OutlineThickness);
 			LastVert = NewVert;
 		}
-		DrawDebugLine(World, LastVert, LocalToWorldTransform.TransformPosition(BoundaryPolygon[0]), OutlineRGB);
+		DrawDebugLine(World, LastVert, LocalToWorldTransform.TransformPosition(BoundaryPolygon[0]), OutlineRGB, PersistForSeconds > 0 ? true : false, PersistForSeconds, OutlineThickness);
 	}
 
 	const FVector WorldSpaceCenter = LocalToWorldTransform.TransformPosition(Center);
