Merging using Dev-VR_to_Release-4.21

diff --git a/Engine/Source/Runtime/AugmentedReality/Public/ARTypes.h b/Engine/Source/Runtime/AugmentedReality/Public/ARTypes.h
--- a/Engine/Source/Runtime/AugmentedReality/Public/ARTypes.h
+++ b/Engine/Source/Runtime/AugmentedReality/Public/ARTypes.h
@@ -343,15 +343,15 @@ public:
 	}
 	
 	/** The desired or supported number of frames per second for this video format */
-	UPROPERTY(BlueprintReadOnly, Category="AR AugmentedReality|Session")
+	UPROPERTY(EditAnywhere, BlueprintReadOnly, Category="AR AugmentedReality|Session")
 	int32 FPS;
 	
 	/** The desired or supported width in pixels for this video format */
-	UPROPERTY(BlueprintReadOnly, Category = "AR AugmentedReality|Session")
+	UPROPERTY(EditAnywhere, BlueprintReadOnly, Category = "AR AugmentedReality|Session")
 	int32 Width;
 
 	/** The desired or supported height in pixels for this video format */
-	UPROPERTY(BlueprintReadOnly, Category = "AR AugmentedReality|Session")
+	UPROPERTY(EditAnywhere, BlueprintReadOnly, Category = "AR AugmentedReality|Session")
 	int32 Height;
 
 	bool IsValidFormat() { return FPS > 0 && Width > 0 && Height > 0; }
