Merging using Dev-VR_to_Release-4.21

diff --git a/Engine/Source/Runtime/AugmentedReality/Public/ARBlueprintLibrary.h b/Engine/Source/Runtime/AugmentedReality/Public/ARBlueprintLibrary.h
--- a/Engine/Source/Runtime/AugmentedReality/Public/ARBlueprintLibrary.h
+++ b/Engine/Source/Runtime/AugmentedReality/Public/ARBlueprintLibrary.h
@@ -179,7 +179,7 @@ public:
 	static bool AddManualEnvironmentCaptureProbe(FVector Location, FVector Extent);
 
 	/** @return the current world mapping status for the AR world */
-	UFUNCTION(BlueprintCallable, Category = "AR AugmentedReality|Tracking", meta = (DisplayName="Get AR World Mapping Status", Keywords = "ar augmentedreality augmented reality tracking anchor"))
+	UFUNCTION(BlueprintPure, Category = "AR AugmentedReality|Tracking", meta = (DisplayName="Get AR World Mapping Status", Keywords = "ar augmentedreality augmented reality tracking anchor"))
 	static EARWorldMappingState GetWorldMappingStatus();
 	
 	/** @return The list of supported video formats for this device */
