Merging using Dev-VR_to_Release-4.21

diff --git a/Engine/Source/Runtime/AugmentedReality/Private/ARSharedWorldGameMode.cpp b/Engine/Source/Runtime/AugmentedReality/Private/ARSharedWorldGameMode.cpp
--- a/Engine/Source/Runtime/AugmentedReality/Private/ARSharedWorldGameMode.cpp
+++ b/Engine/Source/Runtime/AugmentedReality/Private/ARSharedWorldGameMode.cpp
@@ -6,7 +6,7 @@
 
 AARSharedWorldGameMode::AARSharedWorldGameMode(const FObjectInitializer& ObjectInitializer)
 	: Super(ObjectInitializer)
-	, BufferSizePerChunk(16 * 1024)
+	, BufferSizePerChunk(512)
 {
 	GameStateClass = AARSharedWorldGameState::StaticClass();
 	PlayerControllerClass = AARSharedWorldPlayerController::StaticClass();
diff --git a/Engine/Source/Runtime/AugmentedReality/Private/ARSharedWorldGameState.cpp b/Engine/Source/Runtime/AugmentedReality/Private/ARSharedWorldGameState.cpp
--- a/Engine/Source/Runtime/AugmentedReality/Private/ARSharedWorldGameState.cpp
+++ b/Engine/Source/Runtime/AugmentedReality/Private/ARSharedWorldGameState.cpp
@@ -97,7 +97,7 @@ void AARSharedWorldGameState::TriggerCompletionIfDone()
 {
 	if (!bFiredCompletionEvent && ARWorldBytesDelivered == ARWorldBytesTotal && PreviewImageBytesDelivered == PreviewImageBytesTotal)
 	{
-		UE_LOG(LogAR, Log, TEXT("Notifying client AW world data is ready"));
+		UE_LOG(LogAR, Log, TEXT("Notifying client AR world data is ready"));
 
 		bFiredCompletionEvent = true;
 		K2_OnARWorldMapIsReady();
