Merging using Dev-VR_to_Release-4.21

diff --git a/Engine/Source/Runtime/AugmentedReality/Private/ARBlueprintLibrary.cpp b/Engine/Source/Runtime/AugmentedReality/Private/ARBlueprintLibrary.cpp
--- a/Engine/Source/Runtime/AugmentedReality/Private/ARBlueprintLibrary.cpp
+++ b/Engine/Source/Runtime/AugmentedReality/Private/ARBlueprintLibrary.cpp
@@ -25,6 +25,15 @@ EARTrackingQuality UARBlueprintLibrary::GetTrackingQuality()
 
 void UARBlueprintLibrary::StartARSession(UARSessionConfig* SessionConfig)
 {
+	if (SessionConfig == nullptr)
+	{
+#if !(UE_BUILD_SHIPPING || UE_BUILD_TEST)
+		static const TCHAR NoSession_Warning[] = TEXT("Attempting to start an AR session witout a session config object");
+		GEngine->AddOnScreenDebugMessage(INDEX_NONE, 3600.0f, FColor(255,48,16), NoSession_Warning);
+#endif //!(UE_BUILD_SHIPPING || UE_BUILD_TEST)
+		return;
+	}
+
 	static const TCHAR NotARApp_Warning[] = TEXT("Attempting to start an AR session but there is no AR plugin configured. To use AR, enable the proper AR plugin in the Plugin Settings.");
 	
 	auto ARSystem = GetARSystem();
