Merging using Dev-VR_to_Release-4.21

diff --git a/Engine/Plugins/Runtime/AppleARKit/Source/AppleARKit/Private/AppleARKitSystem.cpp b/Engine/Plugins/Runtime/AppleARKit/Source/AppleARKit/Private/AppleARKitSystem.cpp
--- a/Engine/Plugins/Runtime/AppleARKit/Source/AppleARKit/Private/AppleARKitSystem.cpp
+++ b/Engine/Plugins/Runtime/AppleARKit/Source/AppleARKit/Private/AppleARKitSystem.cpp
@@ -839,6 +839,28 @@ TArray<FARVideoFormat> FAppleARKitSystem::OnGetSupportedVideoFormats(EARSessionT
 	return TArray<FARVideoFormat>();
 }
 
+TArray<FVector> FAppleARKitSystem::OnGetPointCloud() const
+{
+	TArray<FVector> PointCloud;
+	
+	if (GameThreadFrame.IsValid())
+	{
+		ARFrame* InARFrame = (ARFrame*)GameThreadFrame->NativeFrame;
+		ARPointCloud* InARPointCloud = InARFrame.rawFeaturePoints;
+		if (InARPointCloud != nullptr)
+		{
+			const int32 Count = InARPointCloud.count;
+			PointCloud.Empty(Count);
+			PointCloud.AddUninitialized(Count);
+			for (int32 Index = 0; Index < Count; Index++)
+			{
+				PointCloud[Index] = FAppleARKitConversion::ToFVector(InARPointCloud.points[Index]);
+			}
+		}
+	}
+	return PointCloud;
+}
+
 #if SUPPORTS_ARKIT_2_0
 /** Since both the object extraction and world saving need to get the world map async, use a common chunk of code for this */
 class FAppleARKitGetWorldMapObjectAsyncTask
diff --git a/Engine/Plugins/Runtime/AppleARKit/Source/AppleARKit/Private/AppleARKitSystem.h b/Engine/Plugins/Runtime/AppleARKit/Source/AppleARKit/Private/AppleARKitSystem.h
--- a/Engine/Plugins/Runtime/AppleARKit/Source/AppleARKit/Private/AppleARKitSystem.h
+++ b/Engine/Plugins/Runtime/AppleARKit/Source/AppleARKit/Private/AppleARKitSystem.h
@@ -84,6 +84,7 @@ protected:
 	virtual TSharedPtr<FARSaveWorldAsyncTask, ESPMode::ThreadSafe> OnSaveWorld() const override;
 	virtual EARWorldMappingState OnGetWorldMappingStatus() const override;
 	virtual TArray<FARVideoFormat> OnGetSupportedVideoFormats(EARSessionType SessionType) const override;
+	virtual TArray<FVector> OnGetPointCloud() const override;
 	//~IARSystemSupport
 
 private:
diff --git a/Engine/Plugins/Runtime/GoogleARCore/Source/GoogleARCoreBase/Private/GoogleARCoreXRTrackingSystem.h b/Engine/Plugins/Runtime/GoogleARCore/Source/GoogleARCoreBase/Private/GoogleARCoreXRTrackingSystem.h
--- a/Engine/Plugins/Runtime/GoogleARCore/Source/GoogleARCoreBase/Private/GoogleARCoreXRTrackingSystem.h
+++ b/Engine/Plugins/Runtime/GoogleARCore/Source/GoogleARCoreBase/Private/GoogleARCoreXRTrackingSystem.h
@@ -68,6 +68,7 @@ protected:
 // @todo -- support these properly
 	virtual EARWorldMappingState OnGetWorldMappingStatus() const override { return EARWorldMappingState::StillMappingNotRelocalizable; }
 	virtual TArray<FARVideoFormat> OnGetSupportedVideoFormats(EARSessionType SessionType) const override { return TArray<FARVideoFormat>(); }
+	virtual TArray<FVector> OnGetPointCloud() const override { return TArray<FVector>(); }
 	//~IARSystemSupport
 
 private:
diff --git a/Engine/Source/Runtime/AugmentedReality/Private/ARBlueprintLibrary.cpp b/Engine/Source/Runtime/AugmentedReality/Private/ARBlueprintLibrary.cpp
--- a/Engine/Source/Runtime/AugmentedReality/Private/ARBlueprintLibrary.cpp
+++ b/Engine/Source/Runtime/AugmentedReality/Private/ARBlueprintLibrary.cpp
@@ -411,6 +411,16 @@ TSharedPtr<FARGetCandidateObjectAsyncTask, ESPMode::ThreadSafe> UARBlueprintLibr
 	return TSharedPtr<FARGetCandidateObjectAsyncTask, ESPMode::ThreadSafe>();
 }
 
+TArray<FVector> UARBlueprintLibrary::GetPointCloud()
+{
+	auto ARSystem = GetARSystem();
+	if (ARSystem.IsValid())
+	{
+		return ARSystem->GetPointCloud();
+	}
+	return TArray<FVector>();
+}
+
 TArray<FARVideoFormat> UARBlueprintLibrary::GetSupportedVideoFormats(EARSessionType SessionType)
 {
 	auto ARSystem = GetARSystem();
diff --git a/Engine/Source/Runtime/AugmentedReality/Private/ARPointCloudComponent.cpp b/Engine/Source/Runtime/AugmentedReality/Private/ARPointCloudComponent.cpp
new file mode 100644
--- /dev/null
+++ b/Engine/Source/Runtime/AugmentedReality/Private/ARPointCloudComponent.cpp
@@ -0,0 +1,177 @@
+// Copyright 1998-2018 Epic Games, Inc. All Rights Reserved.
+
+#include "ARPointCloudComponent.h"
+#include "AugmentedRealityModule.h"
+#include "ARBlueprintLibrary.h"
+#include "PrimitiveViewRelevance.h"
+#include "RenderResource.h"
+#include "RenderingThread.h"
+#include "PrimitiveSceneProxy.h"
+#include "Containers/ResourceArray.h"
+#include "EngineGlobals.h"
+#include "VertexFactory.h"
+#include "MaterialShared.h"
+#include "Materials/Material.h"
+#include "LocalVertexFactory.h"
+#include "Engine/Engine.h"
+#include "SceneManagement.h"
+#include "DynamicMeshBuilder.h"
+#include "StaticMeshResources.h"
+
+DECLARE_CYCLE_STAT(TEXT("Create Point Cloud Proxy"), STAT_PointCloud_CreateSceneProxy, STATGROUP_AR);
+DECLARE_CYCLE_STAT(TEXT("Update Point Cloud GT"), STAT_PointCloud_UpdatePointCloud, STATGROUP_AR);
+DECLARE_CYCLE_STAT(TEXT("Point Cloud Get Mesh Elements"), STAT_PointCloud_GetMeshElements, STATGROUP_AR);
+
+/** Procedural mesh scene proxy */
+class FARPointCloudSceneProxy final :
+	public FPrimitiveSceneProxy
+{
+public:
+	SIZE_T GetTypeHash() const override
+	{
+		static size_t UniquePointer;
+		return reinterpret_cast<size_t>(&UniquePointer);
+	}
+
+	FARPointCloudSceneProxy(UARPointCloudComponent* Component)
+		: FPrimitiveSceneProxy(Component)
+		, Points(Component->PointCloud)
+		, Color(Component->PointColor)
+		, Size(Component->PointSize)
+		, MaterialRelevance(Component->GetMaterialRelevance(GetScene().GetFeatureLevel()))
+	{
+	}
+
+	virtual ~FARPointCloudSceneProxy()
+	{
+	}
+
+	virtual void GetDynamicMeshElements(const TArray<const FSceneView*>& Views, const FSceneViewFamily& ViewFamily, uint32 VisibilityMap, FMeshElementCollector& Collector) const override
+	{
+		SCOPE_CYCLE_COUNTER(STAT_PointCloud_GetMeshElements);
+
+		for (int32 ViewIndex = 0; ViewIndex < Views.Num(); ViewIndex++)
+		{
+			if (VisibilityMap & (1 << ViewIndex))
+			{
+				FPrimitiveDrawInterface* PDI = Collector.GetPDI(ViewIndex);
+				// Loop through manually drawing the points since PT_PointList isn't really a supported primitive type
+				for (const FVector& Point : Points)
+				{
+					PDI->DrawPoint(Point, Color, Size, SDPG_World);
+				}
+			}
+		}
+	}
+
+	virtual FPrimitiveViewRelevance GetViewRelevance(const FSceneView* View) const
+	{
+		FPrimitiveViewRelevance Result;
+		Result.bDrawRelevance = IsShown(View);
+		Result.bShadowRelevance = false;
+		Result.bDynamicRelevance = true;
+		Result.bRenderInMainPass = true;
+		Result.bUsesLightingChannels = false;
+		Result.bRenderCustomDepth = false;
+		MaterialRelevance.SetPrimitiveViewRelevance(Result);
+		return Result;
+	}
+
+	virtual bool CanBeOccluded() const override
+	{
+		return false;
+	}
+
+	virtual uint32 GetMemoryFootprint(void) const
+	{
+		return(sizeof(*this) + GetAllocatedSize());
+	}
+
+	uint32 GetAllocatedSize(void) const
+	{
+		return(FPrimitiveSceneProxy::GetAllocatedSize());
+	}
+
+private:
+	TArray<FVector> Points;
+	FLinearColor Color;
+	float Size;
+	FMaterialRelevance MaterialRelevance;
+};
+
+UARPointCloudComponent::UARPointCloudComponent(const FObjectInitializer& ObjectInitializer)
+	: Super(ObjectInitializer)
+	, PointColor(1.f, 1.f, 1.f, 1.f)
+	, PointSize(4.f)
+{
+	PrimaryComponentTick.bCanEverTick = true;
+	// Tick late in the frame to make sure the ARKit has had a chance to update
+	PrimaryComponentTick.TickGroup = TG_PostPhysics;
+}
+
+void UARPointCloudComponent::InitializeComponent()
+{
+	Super::InitializeComponent();
+	
+	PrimaryComponentTick.SetTickFunctionEnable(bAutoBindToARSystem);
+}
+
+void UARPointCloudComponent::ClearPointCloud()
+{
+	PointCloud.Empty();
+	LocalBounds = FBoxSphereBounds(FVector(0, 0, 0), FVector(0, 0, 0), 0);
+
+	MarkRenderStateDirty();
+}
+
+FPrimitiveSceneProxy* UARPointCloudComponent::CreateSceneProxy()
+{
+	SCOPE_CYCLE_COUNTER(STAT_PointCloud_CreateSceneProxy);
+	return new FARPointCloudSceneProxy(this);
+}
+
+FBoxSphereBounds UARPointCloudComponent::CalcBounds(const FTransform& LocalToWorld) const
+{
+	return LocalBounds;
+}
+
+FMatrix UARPointCloudComponent::GetRenderMatrix() const
+{
+	return FMatrix::Identity;
+}
+
+void UARPointCloudComponent::SetPointCloud(const TArray<FVector>& Points)
+{
+	// Zero and rebuild our bounds from the points
+	FBox PointBounds(Points);
+	LocalBounds = FBoxSphereBounds(PointBounds);
+	PointCloud = Points;
+
+	MarkRenderStateDirty();
+}
+
+void UARPointCloudComponent::SetPointColor(const FLinearColor& Color)
+{
+	PointColor = Color;
+	
+	MarkRenderStateDirty();
+}
+
+void UARPointCloudComponent::SetPointSize(float Size)
+{
+	PointSize = Size;
+	
+	MarkRenderStateDirty();
+}
+
+void UARPointCloudComponent::TickComponent(float DeltaTime, ELevelTick, FActorComponentTickFunction*)
+{
+	SCOPE_CYCLE_COUNTER(STAT_PointCloud_UpdatePointCloud);
+	
+	if (!bAutoBindToARSystem)
+	{
+		return;
+	}
+	// Ask the ar system for updated point cloud
+	SetPointCloud(UARBlueprintLibrary::GetPointCloud());
+}
diff --git a/Engine/Source/Runtime/AugmentedReality/Private/ARSystem.cpp b/Engine/Source/Runtime/AugmentedReality/Private/ARSystem.cpp
--- a/Engine/Source/Runtime/AugmentedReality/Private/ARSystem.cpp
+++ b/Engine/Source/Runtime/AugmentedReality/Private/ARSystem.cpp
@@ -181,3 +181,7 @@ void FARSystemBase::AddReferencedObjects(FReferenceCollector& Collector)
 	}
 }
 
+TArray<FVector> FARSystemBase::GetPointCloud() const
+{
+	return OnGetPointCloud();
+}
diff --git a/Engine/Source/Runtime/AugmentedReality/Private/AugmentedRealityModule.h b/Engine/Source/Runtime/AugmentedReality/Private/AugmentedRealityModule.h
--- a/Engine/Source/Runtime/AugmentedReality/Private/AugmentedRealityModule.h
+++ b/Engine/Source/Runtime/AugmentedReality/Private/AugmentedRealityModule.h
@@ -3,6 +3,7 @@
 #pragma once
 
 #include "CoreMinimal.h"
+#include "Stats/Stats.h"
 
 DECLARE_LOG_CATEGORY_EXTERN(LogAR, Log, All);
-//DECLARE_LOG_CATEGORY_EXTERN(LogAR, Verbose, All);
+DECLARE_STATS_GROUP(TEXT("AR"), STATGROUP_AR, STATCAT_Advanced);
diff --git a/Engine/Source/Runtime/AugmentedReality/Public/ARBlueprintLibrary.h b/Engine/Source/Runtime/AugmentedReality/Public/ARBlueprintLibrary.h
--- a/Engine/Source/Runtime/AugmentedReality/Public/ARBlueprintLibrary.h
+++ b/Engine/Source/Runtime/AugmentedReality/Public/ARBlueprintLibrary.h
@@ -182,6 +182,10 @@ public:
 	UFUNCTION(BlueprintPure, Category = "AR AugmentedReality|Tracking", meta = (DisplayName="Get AR World Mapping Status", Keywords = "ar augmentedreality augmented reality tracking anchor"))
 	static EARWorldMappingState GetWorldMappingStatus();
 	
+	/** @return the raw point cloud view of the AR scene */
+	UFUNCTION(BlueprintPure, Category = "AR AugmentedReality|Tracking", meta = (DisplayName="Get AR Point Cloud", Keywords = "ar augmentedreality augmented reality tracking point cloud"))
+	static TArray<FVector> GetPointCloud();
+
 	/** @return The list of supported video formats for this device */
 	UFUNCTION(BlueprintCallable, Category = "AR AugmentedReality|Capabilities", meta = (DisplayName="Get Supported AR Video Formats", Keywords = "ar augmentedreality augmented reality config video formats"))
 	static TArray<FARVideoFormat> GetSupportedVideoFormats(EARSessionType SessionType);
diff --git a/Engine/Source/Runtime/AugmentedReality/Public/ARPointCloudComponent.h b/Engine/Source/Runtime/AugmentedReality/Public/ARPointCloudComponent.h
new file mode 100644
--- /dev/null
+++ b/Engine/Source/Runtime/AugmentedReality/Public/ARPointCloudComponent.h
@@ -0,0 +1,78 @@
+// Copyright 1998-2018 Epic Games, Inc. All Rights Reserved.
+
+#pragma once
+
+#include "CoreMinimal.h"
+#include "UObject/ObjectMacros.h"
+#include "Components/MeshComponent.h"
+#include "ARPointCloudComponent.generated.h"
+
+class FPrimitiveSceneProxy;
+
+/**
+ * Component for rendering a point cloud
+ */
+UCLASS(hidecategories = (Object, LOD), meta = (BlueprintSpawnableComponent), ClassGroup = AR)
+class AUGMENTEDREALITY_API UARPointCloudComponent :
+	public UMeshComponent
+{
+	GENERATED_UCLASS_BODY()
+
+	/**	If true, each tick the component will try to update its point cloud data from the AR system */
+	UPROPERTY(EditAnywhere, BlueprintReadWrite, Category = "AR|Point Cloud")
+	bool bAutoBindToARSystem;
+	
+	/** Point cloud data that will be used for rendering, assumes each point is in world space */
+	UPROPERTY(BlueprintReadOnly, Category = "AR|Point Cloud")
+	TArray<FVector> PointCloud;
+	
+	/**
+	 * Updates the point cloud data with the new set of points
+	 *
+	 * @param Points the new set of points to use
+	 */
+	UFUNCTION(BlueprintCallable, Category = "AR|Point Cloud")
+	void SetPointCloud(const TArray<FVector>& Points);
+	
+	/**
+	 * Empties the point cloud
+	 */
+	UFUNCTION(BlueprintCallable, Category = "AR|Point Cloud")
+	void ClearPointCloud();
+	
+	/** The color to render the points with */
+	UPROPERTY(EditAnywhere, BlueprintReadOnly, Category = "AR|Point Cloud")
+	FLinearColor PointColor;
+
+	/** Allows you to change the color of the points being rendered */
+	UFUNCTION(BlueprintCallable, Category = "AR|Point Cloud")
+	void SetPointColor(const FLinearColor& Color);
+
+	/** The size of the point when rendering */
+	UPROPERTY(EditAnywhere, BlueprintReadOnly, Category = "AR|Point Cloud")
+	float PointSize;
+	
+	/** Allows you to change the size of the points being rendered */
+	UFUNCTION(BlueprintCallable, Category = "AR|Point Cloud")
+	void SetPointSize(float Size);
+	
+private:
+	virtual FMatrix GetRenderMatrix() const override;
+	virtual FPrimitiveSceneProxy* CreateSceneProxy() override;
+	virtual int32 GetNumMaterials() const override { return 1; }
+	virtual void InitializeComponent() override;
+	virtual FBoxSphereBounds CalcBounds(const FTransform& LocalToWorld) const override;
+
+	virtual void TickComponent(float DeltaTime, enum ELevelTick TickType, FActorComponentTickFunction* ThisTickFunction) override;
+
+	/** World space bounds of the point cloud */
+	UPROPERTY()
+	FBoxSphereBounds LocalBounds;
+	
+	/** Used to know whether to update the point cloud or not */
+	float LastUpdateTimestamp;
+
+	friend class FARPointCloudSceneProxy;
+};
+
+
diff --git a/Engine/Source/Runtime/AugmentedReality/Public/ARSystem.h b/Engine/Source/Runtime/AugmentedReality/Public/ARSystem.h
--- a/Engine/Source/Runtime/AugmentedReality/Public/ARSystem.h
+++ b/Engine/Source/Runtime/AugmentedReality/Public/ARSystem.h
@@ -135,6 +135,9 @@ public:
 	/** @return The list of supported video formats for this device and session type */
 	virtual TArray<FARVideoFormat> OnGetSupportedVideoFormats(EARSessionType SessionType) const = 0;
 	
+	/** @return the current point cloud data for the ar scene */
+	virtual TArray<FVector> OnGetPointCloud() const = 0;
+	
 public:
 	virtual ~IARSystemSupport(){}
 };
@@ -218,6 +221,9 @@ public:
 	/** \see UARBlueprintLibrary::GetSupportedVideoFormats() */
 	TArray<FARVideoFormat> GetSupportedVideoFormats(EARSessionType SessionType = EARSessionType::World) const;
 	
+	/** \see UARBlueprintLibrary::GetPointCloud() */
+	TArray<FVector> GetPointCloud() const;
+
 	virtual void* GetARSessionRawPointer() = 0;
 	virtual void* GetGameThreadARFrameRawPointer() = 0;
 	
