Merging using Dev-VR_to_Release-4.21

diff --git a/Engine/Source/Runtime/AugmentedReality/Private/ARPointCloudComponent.cpp b/Engine/Source/Runtime/AugmentedReality/Private/ARPointCloudComponent.cpp
--- a/Engine/Source/Runtime/AugmentedReality/Private/ARPointCloudComponent.cpp
+++ b/Engine/Source/Runtime/AugmentedReality/Private/ARPointCloudComponent.cpp
@@ -38,6 +38,7 @@ public:
 		, Points(Component->PointCloud)
 		, Color(Component->PointColor)
 		, Size(Component->PointSize)
+		, bIsVisible(Component->bIsVisible)
 		, MaterialRelevance(Component->GetMaterialRelevance(GetScene().GetFeatureLevel()))
 	{
 	}
@@ -49,6 +50,11 @@ public:
 	virtual void GetDynamicMeshElements(const TArray<const FSceneView*>& Views, const FSceneViewFamily& ViewFamily, uint32 VisibilityMap, FMeshElementCollector& Collector) const override
 	{
 		SCOPE_CYCLE_COUNTER(STAT_PointCloud_GetMeshElements);
+		
+		if (!bIsVisible)
+		{
+			return;
+		}
 
 		for (int32 ViewIndex = 0; ViewIndex < Views.Num(); ViewIndex++)
 		{
@@ -67,7 +73,7 @@ public:
 	virtual FPrimitiveViewRelevance GetViewRelevance(const FSceneView* View) const
 	{
 		FPrimitiveViewRelevance Result;
-		Result.bDrawRelevance = IsShown(View);
+		Result.bDrawRelevance = IsShown(View) && bIsVisible;
 		Result.bShadowRelevance = false;
 		Result.bDynamicRelevance = true;
 		Result.bRenderInMainPass = true;
@@ -96,6 +102,7 @@ private:
 	TArray<FVector> Points;
 	FLinearColor Color;
 	float Size;
+	bool bIsVisible;
 	FMaterialRelevance MaterialRelevance;
 };
 
@@ -175,3 +182,47 @@ void UARPointCloudComponent::TickComponent(float DeltaTime, ELevelTick, FActorCo
 	// Ask the ar system for updated point cloud
 	SetPointCloud(UARBlueprintLibrary::GetPointCloud());
 }
+
+TArray<FVector> UARPointCloudComponent::GetPointsInBox(const FBox& WorldSpaceBox) const
+{
+	TArray<FVector> OutPoints;
+	// Miminize the number of allocations by presizing the max
+	OutPoints.Empty(PointCloud.Num());
+
+	for (const FVector& Point : PointCloud)
+	{
+		if (WorldSpaceBox.IsInsideOrOn(Point))
+		{
+			OutPoints.Add(Point);
+		}
+	}
+
+	return OutPoints;
+}
+
+TArray<FVector> UARPointCloudComponent::GetPointsOutsideBox(const FBox& WorldSpaceBox) const
+{
+	TArray<FVector> OutPoints;
+	// Miminize the number of allocations by presizing the max
+	OutPoints.Empty(PointCloud.Num());
+	
+	for (const FVector& Point : PointCloud)
+	{
+		if (!WorldSpaceBox.IsInsideOrOn(Point))
+		{
+			OutPoints.Add(Point);
+		}
+	}
+	
+	return OutPoints;
+}
+
+void UARPointCloudComponent::SetIsVisible(bool bNewVisibity)
+{
+	if (bNewVisibity != bIsVisible)
+	{
+		bIsVisible = bNewVisibity;
+		
+		MarkRenderStateDirty();
+	}
+}
diff --git a/Engine/Source/Runtime/AugmentedReality/Public/ARPointCloudComponent.h b/Engine/Source/Runtime/AugmentedReality/Public/ARPointCloudComponent.h
--- a/Engine/Source/Runtime/AugmentedReality/Public/ARPointCloudComponent.h
+++ b/Engine/Source/Runtime/AugmentedReality/Public/ARPointCloudComponent.h
@@ -22,6 +22,14 @@ class AUGMENTEDREALITY_API UARPointCloudComponent :
 	UPROPERTY(EditAnywhere, BlueprintReadWrite, Category = "AR|Point Cloud")
 	bool bAutoBindToARSystem;
 	
+	/**	If true, each tick the component will render its point cloud */
+	UPROPERTY(EditAnywhere, BlueprintReadWrite, Category = "AR|Point Cloud")
+	bool bIsVisible;
+	
+	/** Changes the visibility setting */
+	UFUNCTION(BlueprintCallable, Category = "AR|Point Cloud")
+	void SetIsVisible(bool bNewVisibility);
+
 	/** Point cloud data that will be used for rendering, assumes each point is in world space */
 	UPROPERTY(BlueprintReadOnly, Category = "AR|Point Cloud")
 	TArray<FVector> PointCloud;
@@ -55,6 +63,14 @@ class AUGMENTEDREALITY_API UARPointCloudComponent :
 	/** Allows you to change the size of the points being rendered */
 	UFUNCTION(BlueprintCallable, Category = "AR|Point Cloud")
 	void SetPointSize(float Size);
+
+	/** Determines which points are within the box and returns those to the caller */
+	UFUNCTION(BlueprintPure, Category = "AR|Point Cloud")
+	TArray<FVector> GetPointsInBox(const FBox& WorldSpaceBox) const;
+
+	/** Determines which points are outside the box and returns those to the caller */
+	UFUNCTION(BlueprintPure, Category = "AR|Point Cloud")
+	TArray<FVector> GetPointsOutsideBox(const FBox& WorldSpaceBox) const;
 	
 private:
 	virtual FMatrix GetRenderMatrix() const override;
