Merging using Dev-VR_to_Release-4.21

diff --git a/Engine/Source/Runtime/AugmentedReality/Private/ARTypes.cpp b/Engine/Source/Runtime/AugmentedReality/Private/ARTypes.cpp
--- a/Engine/Source/Runtime/AugmentedReality/Private/ARTypes.cpp
+++ b/Engine/Source/Runtime/AugmentedReality/Private/ARTypes.cpp
@@ -3,6 +3,7 @@
 
 #include "ARTypes.h"
 #include "HAL/ThreadSafeBool.h"
+#include "Engine/Texture2D.h"
 
 bool FARAsyncTask::HadError() const
 {
@@ -31,3 +32,54 @@ TArray<uint8> FARSaveWorldAsyncTask::GetSavedWorldData()
 	}
 	return TArray<uint8>();
 }
+
+#if WITH_EDITOR
+void UARCandidateImage::PostEditChangeProperty(FPropertyChangedEvent& PropertyChangedEvent)
+{
+	if (PropertyChangedEvent.Property != nullptr)
+	{
+		float AspectRatio = 1.f;
+		if (CandidateTexture != nullptr)
+		{
+			if (Orientation == EARCandidateImageOrientation::Landscape)
+			{
+				AspectRatio = (float)CandidateTexture->GetSizeY() / (float)CandidateTexture->GetSizeX();
+			}
+			else
+			{
+				AspectRatio = (float)CandidateTexture->GetSizeX() / (float)CandidateTexture->GetSizeY();
+			}
+		}
+		// If the texture has changed enforce the aspect ratio on the physical size
+		if (PropertyChangedEvent.Property->GetFName() == FName(TEXT("CandidateTexture")))
+		{
+			Height = Width * AspectRatio;
+		}
+		// Adjust width if they changed the height
+		else if (PropertyChangedEvent.Property->GetFName() == FName(TEXT("Height")))
+		{
+			Width = Height * (1.f / AspectRatio);
+		}
+		// Adjust height if they changed the width
+		else if (PropertyChangedEvent.Property->GetFName() == FName(TEXT("Width")))
+		{
+			Height = Width * AspectRatio;
+		}
+		// Adjust the sizes if they switched the orientation
+		else if (PropertyChangedEvent.Property->GetFName() == FName(TEXT("Orientation")))
+		{
+			if (Orientation == EARCandidateImageOrientation::Landscape)
+			{
+				Width = Height * (1.f / AspectRatio);
+			}
+			else
+			{
+				Height = Width * AspectRatio;
+			}
+		}
+	}
+	
+	Super::PostEditChangeProperty(PropertyChangedEvent);
+}
+#endif
+
diff --git a/Engine/Source/Runtime/AugmentedReality/Public/ARTypes.h b/Engine/Source/Runtime/AugmentedReality/Public/ARTypes.h
--- a/Engine/Source/Runtime/AugmentedReality/Public/ARTypes.h
+++ b/Engine/Source/Runtime/AugmentedReality/Public/ARTypes.h
@@ -191,6 +191,11 @@ public:
 	EARCandidateImageOrientation GetOrientation() const { return Orientation; }
 
 private:
+#if WITH_EDITOR
+	/** Used to enforce physical sizes matching the aspect ratio of the images */
+	virtual void PostEditChangeProperty(FPropertyChangedEvent& PropertyChangedEvent) override;
+#endif
+	
 	/** The image to detect in scenes */
 	UPROPERTY(EditAnywhere, Category = "AR Candidate Image")
 	UTexture2D* CandidateTexture;
