Merging using Dev-VR_to_Release-4.21

diff --git a/Engine/Source/Runtime/AugmentedReality/Private/ARTrackable.cpp b/Engine/Source/Runtime/AugmentedReality/Private/ARTrackable.cpp
--- a/Engine/Source/Runtime/AugmentedReality/Private/ARTrackable.cpp
+++ b/Engine/Source/Runtime/AugmentedReality/Private/ARTrackable.cpp
@@ -170,11 +170,15 @@ void UARPlaneGeometry::DebugDraw( UWorld* World, const FLinearColor& OutlineColo
 void UARTrackedImage::DebugDraw(UWorld* World, const FLinearColor& OutlineColor, float OutlineThickness, float PersistForSeconds /*= 0.0f*/) const
 {
 	const FTransform LocalToWorldTransform = GetLocalToWorldTransform();
-	const FString CurAnchorDebugName = GetDebugName().ToString();
-	const FColor OutlineRGB =OutlineColor.ToFColor(false);
-	ARDebugHelpers::DrawDebugString( World, LocalToWorldTransform.GetLocation(), CurAnchorDebugName, 0.25f*OutlineThickness, OutlineRGB, PersistForSeconds, true);
+	const FString CurAnchorDebugName = FString::Printf(TEXT("%s - %s"), *GetDebugName().ToString(), *DetectedImage->GetFriendlyName());
+	const FColor OutlineRGB = OutlineColor.ToFColor(false);
 
-	DrawDebugPoint(World, LocalToWorldTransform.GetLocation(), 0.5f, OutlineRGB, false, PersistForSeconds, 0);
+	FVector Extent(DetectedImage->GetPhysicalHeight() / 2.f, DetectedImage->GetPhysicalWidth() / 2.f, 0.f);
+	
+	const FVector WorldSpaceCenter = LocalToWorldTransform.GetLocation();
+	DrawDebugBox(World, WorldSpaceCenter, Extent, LocalToWorldTransform.GetRotation(), OutlineRGB, false, PersistForSeconds, 0, 0.1f * OutlineThickness);
+
+	ARDebugHelpers::DrawDebugString(World, WorldSpaceCenter, CurAnchorDebugName, 0.25f * OutlineThickness, OutlineRGB, PersistForSeconds, true);
 }
 
 void UARTrackedImage::UpdateTrackedGeometry(const TSharedRef<FARSystemBase, ESPMode::ThreadSafe>& InTrackingSystem, uint32 FrameNumber, double Timestamp, const FTransform& InLocalToTrackingTransform, const FTransform& InAlignmentTransform, UARCandidateImage* InDetectedImage)
