Changed the AR error message to match the project setting it is referring to

diff --git a/Engine/Source/Runtime/AugmentedReality/Private/ARSystem.cpp b/Engine/Source/Runtime/AugmentedReality/Private/ARSystem.cpp
--- a/Engine/Source/Runtime/AugmentedReality/Private/ARSystem.cpp
+++ b/Engine/Source/Runtime/AugmentedReality/Private/ARSystem.cpp
@@ -41,7 +41,7 @@ EARTrackingQuality FARSystemBase::GetTrackingQuality() const
 
 void FARSystemBase::StartARSession(UARSessionConfig* InSessionConfig)
 {
-	static const TCHAR NotARApp_Warning[] = TEXT("To use AR, enable bIsARApp under Project Settings.");
+	static const TCHAR NotARApp_Warning[] = TEXT("To use AR, enable Supports AR under Project Settings.");
 	
 	const bool bIsARApp = GetDefault<UGeneralProjectSettings>()->bSupportAR;
 	if (ensureAlwaysMsgf(bIsARApp, NotARApp_Warning))
