Added the missing UPROPERTY decoration for the world alignment setting

diff --git a/Engine/Source/Runtime/AugmentedReality/Public/ARSessionConfig.h b/Engine/Source/Runtime/AugmentedReality/Public/ARSessionConfig.h
--- a/Engine/Source/Runtime/AugmentedReality/Public/ARSessionConfig.h
+++ b/Engine/Source/Runtime/AugmentedReality/Public/ARSessionConfig.h
@@ -120,6 +120,7 @@ private:
 	//~ UObject interface
 
 protected:
+	UPROPERTY(EditAnywhere, Category = "AR Settings")
 	/** @see EARWorldAlignment */
 	EARWorldAlignment WorldAlignment;
 
