Added a file that I missed when merging

diff --git a/Engine/Source/Runtime/AugmentedReality/Private/ARBlueprintProxy.cpp b/Engine/Source/Runtime/AugmentedReality/Private/ARBlueprintProxy.cpp
new file mode 100644
--- /dev/null
+++ b/Engine/Source/Runtime/AugmentedReality/Private/ARBlueprintProxy.cpp
@@ -0,0 +1,121 @@
+//@joeg -- ARkit 2.0 additions
+
+// Copyright 1998-2018 Epic Games, Inc. All Rights Reserved.
+
+#include "ARBlueprintProxy.h"
+#include "ARSystem.h"
+
+TSharedPtr<FARSystemBase, ESPMode::ThreadSafe> UARBaseAsyncTaskBlueprintProxy::RegisteredARSystem = nullptr;
+
+UARBaseAsyncTaskBlueprintProxy::UARBaseAsyncTaskBlueprintProxy(const FObjectInitializer& ObjectInitializer)
+	: Super(ObjectInitializer)
+	, bShouldTick(true)
+{
+}
+
+void UARBaseAsyncTaskBlueprintProxy::Tick(float DeltaTime)
+{
+	if (!AsyncTask.IsValid())
+	{
+		bShouldTick = false;
+		ReportFailure();
+		return;
+	}
+
+	if (AsyncTask->IsDone())
+	{
+		bShouldTick = false;
+		// Fire the right delegate
+		if (!AsyncTask->HadError())
+		{
+			ReportSuccess();
+		}
+		else
+		{
+			ReportFailure();
+		}
+	}
+}
+
+void UARBaseAsyncTaskBlueprintProxy::RegisterAsARSystem(const TSharedPtr<FARSystemBase, ESPMode::ThreadSafe>& NewARSystem)
+{
+	RegisteredARSystem = NewARSystem;
+}
+
+
+const TSharedPtr<FARSystemBase, ESPMode::ThreadSafe>& UARBaseAsyncTaskBlueprintProxy::GetARSystem()
+{
+	return RegisteredARSystem;
+}
+
+UARSaveWorldAsyncTaskBlueprintProxy::UARSaveWorldAsyncTaskBlueprintProxy(const FObjectInitializer& ObjectInitializer)
+	: Super(ObjectInitializer)
+{
+}
+
+UARSaveWorldAsyncTaskBlueprintProxy* UARSaveWorldAsyncTaskBlueprintProxy::CreateProxyObjectForARSaveWorld(bool bCompressData)
+{
+	UARSaveWorldAsyncTaskBlueprintProxy* Proxy = NewObject<UARSaveWorldAsyncTaskBlueprintProxy>();
+	Proxy->SetFlags(RF_StrongRefOnFrame);
+
+	auto ARSystem = GetARSystem();
+	if (ARSystem.IsValid())
+	{
+		Proxy->SaveWorldTask = ARSystem->SaveWorld();
+		Proxy->AsyncTask = Proxy->SaveWorldTask;
+	}
+	else
+	{
+		Proxy->AsyncTask = MakeShared<FARErrorSaveWorldAsyncTask, ESPMode::ThreadSafe>(TEXT("SaveWorld - requires a valid, running ARKit 2.0 session"));
+	}
+
+	return Proxy;
+}
+
+void UARSaveWorldAsyncTaskBlueprintProxy::ReportSuccess()
+{
+	SaveWorldResult.WorldData = SaveWorldTask->GetSavedWorldData();
+	OnSuccess.Broadcast(SaveWorldResult);
+}
+
+void UARSaveWorldAsyncTaskBlueprintProxy::ReportFailure()
+{
+	SaveWorldResult.Error = AsyncTask->GetErrorString();
+	OnFailure.Broadcast(SaveWorldResult);
+}
+
+UARGetCandidateObjectAsyncTaskBlueprintProxy::UARGetCandidateObjectAsyncTaskBlueprintProxy(const FObjectInitializer& ObjectInitializer)
+	: Super(ObjectInitializer)
+{
+}
+
+UARGetCandidateObjectAsyncTaskBlueprintProxy* UARGetCandidateObjectAsyncTaskBlueprintProxy::CreateProxyObjectForARGetCandidateObject(FVector Location, FVector Extent)
+{
+	UARGetCandidateObjectAsyncTaskBlueprintProxy* Proxy = NewObject<UARGetCandidateObjectAsyncTaskBlueprintProxy>();
+	Proxy->SetFlags(RF_StrongRefOnFrame);
+
+	auto ARSystem = GetARSystem();
+	if (ARSystem.IsValid())
+	{
+		Proxy->CandidateObjectTask = ARSystem->GetCandidateObject(Location, Extent);
+		Proxy->AsyncTask = Proxy->CandidateObjectTask;
+	}
+	else
+	{
+		Proxy->AsyncTask = MakeShared<FARErrorGetCandidateObjectAsyncTask, ESPMode::ThreadSafe>(TEXT("GetCandidateObject - requires a valid, running ARKit 2.0 session"));
+	}
+
+	return Proxy;
+}
+
+void UARGetCandidateObjectAsyncTaskBlueprintProxy::ReportSuccess()
+{
+	CandidateObjectResult.CandidateObject = CandidateObjectTask->GetCandidateObject();
+	OnSuccess.Broadcast(CandidateObjectResult);
+}
+
+void UARGetCandidateObjectAsyncTaskBlueprintProxy::ReportFailure()
+{
+	CandidateObjectResult.Error = AsyncTask->GetErrorString();
+	OnFailure.Broadcast(CandidateObjectResult);
+}
