Added ARSkyLight class to automatically rebuild the sky cube maps as AR environment maps get updated

diff --git a/Engine/Source/Runtime/AugmentedReality/Private/ARSkyLight.cpp b/Engine/Source/Runtime/AugmentedReality/Private/ARSkyLight.cpp
new file mode 100644
--- /dev/null
+++ b/Engine/Source/Runtime/AugmentedReality/Private/ARSkyLight.cpp
@@ -0,0 +1,62 @@
+// Copyright 1998-2018 Epic Games, Inc. All Rights Reserved.
+
+#include "ARSkyLight.h"
+#include "ARTrackable.h"
+#include "ARTextures.h"
+
+#define LOCTEXT_NAMESPACE "ARSkyLight"
+
+
+AARSkyLight::AARSkyLight(const FObjectInitializer& ObjectInitializer)
+	: Super(ObjectInitializer)
+	, CaptureProbe(nullptr)
+	, LastUpdateTimestamp(0.f)
+{
+	PrimaryActorTick.bCanEverTick = true;
+	// Don't start ticking until we have a valid capture probe to update with
+	PrimaryActorTick.bStartWithTickEnabled = false;
+	// The environment probes don't update often so don't check very often
+	PrimaryActorTick.TickInterval = 0.25f;
+
+	// We only work with environment probes that generate a cube map
+	GetLightComponent()->SourceType = ESkyLightSourceType::SLS_SpecifiedCubemap;
+}
+
+// @todo - JoeG add double buffering of the cube map textures so we can blend the updates rather than popping
+void AARSkyLight::Tick(float DeltaTime)
+{
+	if (CaptureProbe == nullptr)
+	{
+		// Disable our tick because there's nothing to check
+		SetActorTickEnabled(false);
+		return;
+	}
+
+	// Check to see if the probe was updated by the ARSystem
+	if (LastUpdateTimestamp < CaptureProbe->GetLastUpdateTimestamp())
+	{
+		// Trigger a refresh of the cube map data
+		GetLightComponent()->MarkRenderStateDirty();
+		GetLightComponent()->SetCaptureIsDirty();
+		LastUpdateTimestamp = CaptureProbe->GetLastUpdateTimestamp();
+	}
+}
+
+void AARSkyLight::SetEnvironmentCaptureProbe(UAREnvironmentCaptureProbe* InCaptureProbe)
+{
+	CaptureProbe = InCaptureProbe;
+	// Turn ticking on if we have a capture probe that needs regular updating
+	SetActorTickEnabled(CaptureProbe != nullptr);
+
+	if (CaptureProbe != nullptr)
+	{
+		UTextureCube* CubeTexture = CaptureProbe->GetEnvironmentCaptureTexture();
+		GetLightComponent()->SetCubemap(CubeTexture);
+		LastUpdateTimestamp = CaptureProbe->GetLastUpdateTimestamp();
+	}
+	else
+	{
+		LastUpdateTimestamp = 0.f;
+	}
+}
+
diff --git a/Engine/Source/Runtime/AugmentedReality/Private/ARTrackable.cpp b/Engine/Source/Runtime/AugmentedReality/Private/ARTrackable.cpp
--- a/Engine/Source/Runtime/AugmentedReality/Private/ARTrackable.cpp
+++ b/Engine/Source/Runtime/AugmentedReality/Private/ARTrackable.cpp
@@ -276,7 +276,7 @@ FVector UAREnvironmentCaptureProbe::GetExtent() const
 	return Extent;
 }
 
-UAREnvironmentCaptureProbeTexture* UAREnvironmentCaptureProbe::GetEnvironmentCaptureTexture() const
+UAREnvironmentCaptureProbeTexture* UAREnvironmentCaptureProbe::GetEnvironmentCaptureTexture()
 {
 	return EnvironmentCaptureTexture;
 }
diff --git a/Engine/Source/Runtime/AugmentedReality/Public/ARSkyLight.h b/Engine/Source/Runtime/AugmentedReality/Public/ARSkyLight.h
new file mode 100644
--- /dev/null
+++ b/Engine/Source/Runtime/AugmentedReality/Public/ARSkyLight.h
@@ -0,0 +1,30 @@
+// Copyright 1998-2017 Epic Games, Inc. All Rights Reserved.
+
+#pragma once
+
+#include "Components/SkyLightComponent.h"
+#include "Engine/SkyLight.h"
+
+#include "ARSkyLight.generated.h"
+
+/** This sky light class forces a refresh of the cube map data when an AR environment probe changes */
+UCLASS()
+class AUGMENTEDREALITY_API AARSkyLight :
+	public ASkyLight
+{
+	GENERATED_UCLASS_BODY()
+
+public:
+	/** Sets the environment capture probe that this sky light is driven by */
+	UFUNCTION(BlueprintCallable, Category="AR AugmentedReality|SkyLight")
+	void SetEnvironmentCaptureProbe(UAREnvironmentCaptureProbe* InCaptureProbe);
+
+private:
+	virtual void Tick(float DeltaTime) override;
+
+	UPROPERTY()
+	UAREnvironmentCaptureProbe* CaptureProbe;
+
+	/** The timestamp from the environment probe when we last updated the cube map */
+	float LastUpdateTimestamp;
+};
diff --git a/Engine/Source/Runtime/AugmentedReality/Public/ARTrackable.h b/Engine/Source/Runtime/AugmentedReality/Public/ARTrackable.h
--- a/Engine/Source/Runtime/AugmentedReality/Public/ARTrackable.h
+++ b/Engine/Source/Runtime/AugmentedReality/Public/ARTrackable.h
@@ -318,7 +318,7 @@ public:
 	FVector GetExtent() const;
 	/** @see EnvironmentCaptureTexture */
 	UFUNCTION(BlueprintPure, Category="AR AugmentedReality|Environment Capture Probe")
-	UAREnvironmentCaptureProbeTexture* GetEnvironmentCaptureTexture() const;
+	UAREnvironmentCaptureProbeTexture* GetEnvironmentCaptureTexture();
 
 protected:
 	/** The size of area this probe covers */
