Added support to stream a preview image and the accompanying AR world data for shared AR experiences

diff --git a/Engine/Source/Runtime/AugmentedReality/Private/ARSharedWorldGameMode.cpp b/Engine/Source/Runtime/AugmentedReality/Private/ARSharedWorldGameMode.cpp
new file mode 100644
--- /dev/null
+++ b/Engine/Source/Runtime/AugmentedReality/Private/ARSharedWorldGameMode.cpp
@@ -0,0 +1,115 @@
+// Copyright 1998-2018 Epic Games, Inc. All Rights Reserved.
+
+#include "ARSharedWorldGameMode.h"
+#include "Engine/World.h"
+#include "AugmentedRealityModule.h"
+
+AARSharedWorldGameMode::AARSharedWorldGameMode(const FObjectInitializer& ObjectInitializer)
+	: Super(ObjectInitializer)
+	, BufferSizePerChunk(16 * 1024)
+{
+	GameStateClass = AARSharedWorldGameState::StaticClass();
+	PlayerControllerClass = AARSharedWorldPlayerController::StaticClass();
+}
+
+void AARSharedWorldGameMode::SetPreviewImageData(TArray<uint8> ImageData)
+{
+	GetARSharedWorldGameState()->ARSharedWorld.PreviewImageData = MoveTemp(ImageData);
+}
+
+void AARSharedWorldGameMode::SetARSharedWorldData(TArray<uint8> ARWorldData)
+{
+	GetARSharedWorldGameState()->ARSharedWorld.ARWorldData = MoveTemp(ARWorldData);
+}
+
+AARSharedWorldGameState* AARSharedWorldGameMode::GetARSharedWorldGameState()
+{
+	return CastChecked<AARSharedWorldGameState>(GameState);
+}
+
+void AARSharedWorldGameMode::SetARWorldSharingIsReady()
+{
+	// We should only send the world once per game session
+	if (!bShouldSendSharedWorldData)
+	{
+		bShouldSendSharedWorldData = true;
+
+		const FARSharedWorld& SharedWorld = GetARSharedWorldGameState()->ARSharedWorld;
+		UE_LOG(LogAR, Log, TEXT("Ready to share AR data with clients. AR world size is (%d) and preview image size is (%d)"), SharedWorld.ARWorldData.Num(), SharedWorld.PreviewImageData.Num());
+	}
+}
+
+void AARSharedWorldGameMode::Tick(float DeltaSeconds)
+{
+	Super::Tick(DeltaSeconds);
+	
+	if (bShouldSendSharedWorldData)
+	{
+		const FARSharedWorld& SharedWorld = GetARSharedWorldGameState()->ARSharedWorld;
+		
+		// For each player, send them their next chunk of data if needed
+		for (auto It = GetWorld()->GetControllerIterator(); It; ++It)
+		{
+			AARSharedWorldPlayerController* PC = Cast<AARSharedWorldPlayerController>(It->Get());
+			if (PC != nullptr && !PC->IsLocalController() && PC->IsReadyToReceive())
+			{
+				if (PlayerToReplicationStateMap.Contains(PC))
+				{
+					FARSharedWorldReplicationState& ReplState = PlayerToReplicationStateMap[PC];
+					// See if we need to send any image preview data
+					int32 CurrentOffset = ReplState.PreviewImageOffset;
+					if (CurrentOffset < SharedWorld.PreviewImageData.Num())
+					{
+						SendBuffer.Reset(BufferSizePerChunk);
+						
+						// Figure how much needs to be sent (could be less than the chunk size)
+						int32 BytesToSend = FMath::Min(BufferSizePerChunk, SharedWorld.PreviewImageData.Num() - CurrentOffset);
+						SendBuffer.AddUninitialized(BytesToSend);
+						FMemory::Memcpy((void*)SendBuffer.GetData(), (void*)(SharedWorld.PreviewImageData.GetData() + CurrentOffset), BytesToSend);
+						
+						PC->ClientUpdatePreviewImageData(CurrentOffset, SendBuffer);
+						ReplState.PreviewImageOffset += BytesToSend;
+
+						UE_LOG(LogAR, Verbose, TEXT("Sent ClientUpdatePreviewImageData(%d, %d) to PC (%s)"), CurrentOffset, BytesToSend, *PC->GetName());
+					}
+					// See if we need to send any AR world data
+					CurrentOffset = ReplState.ARWorldOffset;
+					if (CurrentOffset < SharedWorld.ARWorldData.Num())
+					{
+						SendBuffer.Reset(BufferSizePerChunk);
+						
+						// Figure how much needs to be sent (could be less than the chunk size)
+						int32 BytesToSend = FMath::Min(BufferSizePerChunk, SharedWorld.ARWorldData.Num() - CurrentOffset);
+						SendBuffer.AddUninitialized(BytesToSend);
+						FMemory::Memcpy((void*)SendBuffer.GetData(), (void*)(SharedWorld.ARWorldData.GetData() + CurrentOffset), BytesToSend);
+						
+						PC->ClientUpdateARWorldData(CurrentOffset, SendBuffer);
+						ReplState.ARWorldOffset += BytesToSend;
+
+						UE_LOG(LogAR, Verbose, TEXT("Sent ClientUpdateARWorldData(%d, %d) to PC (%s)"), CurrentOffset, BytesToSend, *PC->GetName());
+					}
+				}
+				else
+				{
+					// Add if we haven't seen this player before
+					PlayerToReplicationStateMap.Add(PC, FARSharedWorldReplicationState());
+					PC->ClientInitSharedWorld(SharedWorld.PreviewImageData.Num(), SharedWorld.ARWorldData.Num());
+					UE_LOG(LogAR, Verbose, TEXT("Sent InitSharedWorld(%d, %d) to PC (%s)"), SharedWorld.PreviewImageData.Num(), SharedWorld.ARWorldData.Num(), *PC->GetName());
+				}
+			}
+		}
+	}
+}
+
+void AARSharedWorldGameMode::Logout(AController* Exiting)
+{
+	Super::Logout(Exiting);
+	
+	AARSharedWorldPlayerController* PC = Cast<AARSharedWorldPlayerController>(Exiting);
+	if (PC != nullptr && PlayerToReplicationStateMap.Contains(PC))
+	{
+		UE_LOG(LogAR, Verbose, TEXT("Removing PC (%s)"), *PC->GetName());
+
+		PlayerToReplicationStateMap.Remove(PC);
+	}
+}
diff --git a/Engine/Source/Runtime/AugmentedReality/Private/ARSharedWorldGameState.cpp b/Engine/Source/Runtime/AugmentedReality/Private/ARSharedWorldGameState.cpp
new file mode 100644
--- /dev/null
+++ b/Engine/Source/Runtime/AugmentedReality/Private/ARSharedWorldGameState.cpp
@@ -0,0 +1,105 @@
+// Copyright 1998-2018 Epic Games, Inc. All Rights Reserved.
+
+#include "ARSharedWorldGameState.h"
+#include "AugmentedRealityModule.h"
+
+AARSharedWorldGameState::AARSharedWorldGameState(const FObjectInitializer& ObjectInitializer)
+	: Super(ObjectInitializer)
+	, PreviewImageBytesTotal(0)
+	, ARWorldBytesTotal(0)
+	, PreviewImageBytesDelivered(0)
+	, ARWorldBytesDelivered(0)
+	, bFiredCompletionEvent(false)
+{
+	
+}
+
+void AARSharedWorldGameState::InitSharedWorld(int32 PreviewImageSize, int32 ARWorldDataSize)
+{
+	// Should be called on the client only
+	if (GIsServer)
+	{
+		UE_LOG(LogAR, Warning, TEXT("AARSharedWorldGameState::InitSharedWorld() was called on the server. This is client only"));
+		return;
+	}
+
+	ARWorldBytesTotal = 0;
+	PreviewImageBytesTotal = 0;
+	ARWorldBytesDelivered = 0;
+	PreviewImageBytesDelivered = 0;
+
+
+	if (PreviewImageSize > 0 && ARWorldDataSize > 0)
+	{
+		ARSharedWorld.PreviewImageData.Reset(PreviewImageSize);
+		ARSharedWorld.PreviewImageData.AddUninitialized(PreviewImageSize);
+		PreviewImageBytesTotal = PreviewImageSize;
+		
+		ARSharedWorld.ARWorldData.Reset(ARWorldDataSize);
+		ARSharedWorld.ARWorldData.AddUninitialized(ARWorldDataSize);
+		ARWorldBytesTotal = ARWorldDataSize;
+	}
+	else
+	{
+		UE_LOG(LogAR, Warning, TEXT("AARSharedWorldGameState::InitSharedWorld() was called with invalid sizes (%d, %d)"), PreviewImageSize, ARWorldDataSize);
+	}
+}
+
+void AARSharedWorldGameState::UpdatePreviewImageData(int32 Offset, const uint8* Buffer, int32 Size)
+{
+	// Should be called on the client only
+	if (GIsServer)
+	{
+		UE_LOG(LogAR, Warning, TEXT("AARSharedWorldGameState::UpdatePreviewImageData() was called on the server. This is client only"));
+		return;
+	}
+	
+	if (Offset >= 0 && (Offset + Size) <= PreviewImageBytesTotal)
+	{
+		uint8* PreviewImageBuffer = ARSharedWorld.PreviewImageData.GetData();
+		FMemory::Memcpy((void*)&PreviewImageBuffer[Offset], (void*)Buffer, Size);
+		PreviewImageBytesDelivered += Size;
+
+		// Trigger the completion event
+		TriggerCompletionIfDone();
+	}
+	else
+	{
+		UE_LOG(LogAR, Warning, TEXT("AARSharedWorldGameState::UpdatePreviewImageData() was called with bad offset (%d) or size (%d)"), Offset, Size);
+	}
+}
+
+void AARSharedWorldGameState::UpdateARWorldData(int32 Offset, const uint8* Buffer, int32 Size)
+{
+	// Should be called on the client only
+	if (GIsServer)
+	{
+		UE_LOG(LogAR, Warning, TEXT("AARSharedWorldGameState::UpdateARWorldData() was called on the server. This is client only"));
+		return;
+	}
+	
+	if (Offset >= 0 && (Offset + Size) <= ARWorldBytesTotal)
+	{
+		uint8* ARWorldBuffer = ARSharedWorld.ARWorldData.GetData();
+		FMemory::Memcpy((void*)&ARWorldBuffer[Offset], (void*)Buffer, Size);
+		ARWorldBytesDelivered += Size;
+		
+		// Trigger the completion event
+		TriggerCompletionIfDone();
+	}
+	else
+	{
+		UE_LOG(LogAR, Warning, TEXT("AARSharedWorldGameState::UpdateARWorldData() was called with bad offset (%d) or size (%d)"), Offset, Size);
+	}
+}
+
+void AARSharedWorldGameState::TriggerCompletionIfDone()
+{
+	if (!bFiredCompletionEvent && ARWorldBytesDelivered == ARWorldBytesTotal && PreviewImageBytesDelivered == PreviewImageBytesTotal)
+	{
+		UE_LOG(LogAR, Log, TEXT("Notifying client AW world data is ready"));
+
+		bFiredCompletionEvent = true;
+		K2_OnARWorldMapIsReady();
+	}
+}
diff --git a/Engine/Source/Runtime/AugmentedReality/Private/ARSharedWorldPlayerController.cpp b/Engine/Source/Runtime/AugmentedReality/Private/ARSharedWorldPlayerController.cpp
new file mode 100644
--- /dev/null
+++ b/Engine/Source/Runtime/AugmentedReality/Private/ARSharedWorldPlayerController.cpp
@@ -0,0 +1,88 @@
+// Copyright 1998-2018 Epic Games, Inc. All Rights Reserved.
+
+#include "ARSharedWorldPlayerController.h"
+#include "ARSharedWorldGameState.h"
+#include "Engine/World.h"
+#include "AugmentedRealityModule.h"
+
+AARSharedWorldPlayerController::AARSharedWorldPlayerController(const FObjectInitializer& ObjectInitializer)
+	: Super(ObjectInitializer)
+	, bIsReadyToReceive(false)
+{
+	
+}
+
+AARSharedWorldGameState* AARSharedWorldPlayerController::GetGameState()
+{
+	UWorld* World = GetWorld();
+	AARSharedWorldGameState* GameState = World != nullptr ? CastChecked<AARSharedWorldGameState>(World->GetGameState()) : nullptr;
+	return GameState;
+}
+
+bool AARSharedWorldPlayerController::IsGameStateReady()
+{
+	UWorld* World = GetWorld();
+	AARSharedWorldGameState* GameState = World != nullptr ? Cast<AARSharedWorldGameState>(World->GetGameState()) : nullptr;
+	return GameState != nullptr;
+}
+
+void AARSharedWorldPlayerController::Tick(float DeltaSeconds)
+{
+	Super::Tick(DeltaSeconds);
+
+	// If we are the client, check that our game state has replicated over
+	if (!HasAuthority() && !bIsReadyToReceive && IsGameStateReady())
+	{
+		bIsReadyToReceive = true;
+		ServerMarkReadyForReceiving();
+
+		UE_LOG(LogAR, Verbose, TEXT("Notifying server ready to receive via ServerMarkReadyForReceiving()"));
+	}
+}
+
+bool AARSharedWorldPlayerController::ServerMarkReadyForReceiving_Validate()
+{
+	return true;
+}
+
+void AARSharedWorldPlayerController::ServerMarkReadyForReceiving_Implementation()
+{
+	bIsReadyToReceive = true;
+	UE_LOG(LogAR, Verbose, TEXT("Client is ready to receive"));
+}
+
+bool AARSharedWorldPlayerController::ClientInitSharedWorld_Validate(int32 PreviewImageSize, int32 ARWorldDataSize)
+{
+	return PreviewImageSize >= 0 && PreviewImageSize <= (8 * 1024 * 1024) && ARWorldDataSize >= 0 && ARWorldDataSize <= (128 * 1024 * 1024);
+}
+
+void AARSharedWorldPlayerController::ClientInitSharedWorld_Implementation(int32 PreviewImageSize, int32 ARWorldDataSize)
+{
+	UE_LOG(LogAR, Verbose, TEXT("Client received ClientInitSharedWorld(%d, %d)"), PreviewImageSize, ARWorldDataSize);
+
+	GetGameState()->InitSharedWorld(PreviewImageSize, ARWorldDataSize);
+}
+
+bool AARSharedWorldPlayerController::ClientUpdatePreviewImageData_Validate(int32 Offset, const TArray<uint8>& Buffer)
+{
+	return Offset >= 0;
+}
+
+void AARSharedWorldPlayerController::ClientUpdatePreviewImageData_Implementation(int32 Offset, const TArray<uint8>& Buffer)
+{
+	UE_LOG(LogAR, Verbose, TEXT("Client received ClientUpdatePreviewImageData(%d, %d)"), Offset, Buffer.Num());
+
+	GetGameState()->UpdatePreviewImageData(Offset, Buffer.GetData(), Buffer.Num());
+}
+
+bool AARSharedWorldPlayerController::ClientUpdateARWorldData_Validate(int32 Offset, const TArray<uint8>& Buffer)
+{
+	return Offset >= 0;
+}
+
+void AARSharedWorldPlayerController::ClientUpdateARWorldData_Implementation(int32 Offset, const TArray<uint8>& Buffer)
+{
+	UE_LOG(LogAR, Verbose, TEXT("Client received ClientUpdateARWorldData(%d, %d)"), Offset, Buffer.Num());
+
+	GetGameState()->UpdateARWorldData(Offset, Buffer.GetData(), Buffer.Num());
+}
diff --git a/Engine/Source/Runtime/AugmentedReality/Private/AugmentedRealityModule.cpp b/Engine/Source/Runtime/AugmentedReality/Private/AugmentedRealityModule.cpp
--- a/Engine/Source/Runtime/AugmentedReality/Private/AugmentedRealityModule.cpp
+++ b/Engine/Source/Runtime/AugmentedReality/Private/AugmentedRealityModule.cpp
@@ -8,3 +8,5 @@ class FAugmentedRealityModule : public IModuleInterface
 };
 
 IMPLEMENT_MODULE(FAugmentedRealityModule, AugmentedReality);
+
+DEFINE_LOG_CATEGORY(LogAR);
diff --git a/Engine/Source/Runtime/AugmentedReality/Private/AugmentedRealityModule.h b/Engine/Source/Runtime/AugmentedReality/Private/AugmentedRealityModule.h
--- a/Engine/Source/Runtime/AugmentedReality/Private/AugmentedRealityModule.h
+++ b/Engine/Source/Runtime/AugmentedReality/Private/AugmentedRealityModule.h
@@ -2,4 +2,7 @@
 
 #pragma once
 
+#include "CoreMinimal.h"
 
+DECLARE_LOG_CATEGORY_EXTERN(LogAR, Log, All);
+//DECLARE_LOG_CATEGORY_EXTERN(LogAR, Verbose, All);
diff --git a/Engine/Source/Runtime/AugmentedReality/Public/ARSharedWorldGameMode.h b/Engine/Source/Runtime/AugmentedReality/Public/ARSharedWorldGameMode.h
new file mode 100644
--- /dev/null
+++ b/Engine/Source/Runtime/AugmentedReality/Public/ARSharedWorldGameMode.h
@@ -0,0 +1,83 @@
+// Copyright 1998-2018 Epic Games, Inc. All Rights Reserved.
+
+#pragma once
+
+#include "GameFramework/GameMode.h"
+
+#include "ARSharedWorldGameState.h"
+#include "ARSharedWorldPlayerController.h"
+
+#include "ARSharedWorldGameMode.generated.h"
+
+/** Per player information about what data has been sent to them */
+USTRUCT(BlueprintType)
+struct AUGMENTEDREALITY_API FARSharedWorldReplicationState
+{
+	GENERATED_BODY()
+	
+	FARSharedWorldReplicationState()
+	{
+		PreviewImageOffset = ARWorldOffset = 0;
+	}
+	
+	/** The offset in the overall image data buffer */
+	UPROPERTY(BlueprintReadOnly, Category="AR Shared World")
+	int32 PreviewImageOffset;
+	
+	/** The offset in the overall ARWorld data buffer */
+	UPROPERTY(BlueprintReadOnly, Category="AR Shared World")
+	int32 ARWorldOffset;
+};
+
+UCLASS(BlueprintType)
+class AUGMENTEDREALITY_API AARSharedWorldGameMode :
+	public AGameMode
+{
+	GENERATED_UCLASS_BODY()
+
+public:
+	/** The size of the buffer to use per send request. Must be between 1 and 65535, though should not be max to avoid saturation */
+	UPROPERTY(Config, EditAnywhere, BlueprintReadWrite, Category="AR Shared World")
+	int32 BufferSizePerChunk;
+	
+	/**
+	 * Sets the image data for the shared world game session
+	 *
+	 * @param ImageData the blob to use as the image data
+	 */
+	UFUNCTION(BlueprintCallable, BlueprintAuthorityOnly, meta=(DisplayName="Set AR Preview Image Data"), Category="AR Shared World")
+	void SetPreviewImageData(TArray<uint8> ImageData);
+	
+	/**
+	 * Sets the image data for the shared world game session
+	 *
+	 * @param ARWorldData the blob to use as the AR world data
+	 */
+	UFUNCTION(BlueprintCallable, BlueprintAuthorityOnly, meta=(DisplayName="Set AR Shared World Data"), Category="AR Shared World")
+	void SetARSharedWorldData(TArray<uint8> ARWorldData);
+	
+	/**
+	 * Tells the game mode that the AR data is ready and should be replicated to all connected clients
+	 */
+	UFUNCTION(BlueprintCallable, BlueprintAuthorityOnly, meta=(DisplayName="Set AR World Sharing Is Ready"), Category="AR Shared World")
+	void SetARWorldSharingIsReady();
+	
+	/**
+	 * @return the game state for this game mode
+	 */
+	UFUNCTION(BlueprintCallable, BlueprintAuthorityOnly, Category="AR Shared World")
+	AARSharedWorldGameState* GetARSharedWorldGameState();
+	
+private:
+	virtual void Tick(float DeltaSeconds) override;
+	virtual void Logout(AController* Exiting) override;
+
+	/** Tracks whether the data should be sent to all clients or not */
+	bool bShouldSendSharedWorldData;
+	
+	/** Holds the progress for each player that is being replicated to */
+	TMap<AARSharedWorldPlayerController*, FARSharedWorldReplicationState> PlayerToReplicationStateMap;
+	
+	/** Scratch buffer for sending chunks, cached to avoid allocation churn */
+	TArray<uint8> SendBuffer;
+};
diff --git a/Engine/Source/Runtime/AugmentedReality/Public/ARSharedWorldGameState.h b/Engine/Source/Runtime/AugmentedReality/Public/ARSharedWorldGameState.h
new file mode 100644
--- /dev/null
+++ b/Engine/Source/Runtime/AugmentedReality/Public/ARSharedWorldGameState.h
@@ -0,0 +1,81 @@
+// Copyright 1998-2018 Epic Games, Inc. All Rights Reserved.
+
+#pragma once
+
+#include "CoreMinimal.h"
+
+#include "GameFramework/GameState.h"
+#include "ARSharedWorldGameState.generated.h"
+
+USTRUCT(BlueprintType)
+struct AUGMENTEDREALITY_API FARSharedWorld
+{
+	GENERATED_BODY()
+
+	/** The image taken at the time of world saving for use when aligning the AR world later in the session */
+	TArray<uint8> PreviewImageData;
+	/** The snapshot of the AR world that is to be shared */
+	TArray<uint8> ARWorldData;
+};
+
+UCLASS(BlueprintType)
+class AUGMENTEDREALITY_API AARSharedWorldGameState :
+	public AGameState
+{
+	GENERATED_UCLASS_BODY()
+
+public:
+	/** Each client and the host have a copy of the shared world data */
+	UPROPERTY(BlueprintReadOnly, Category="AR Shared World")
+	FARSharedWorld ARSharedWorld;
+
+	/** The size of the image that will be replicated to each client */
+	UPROPERTY(BlueprintReadOnly, Category="AR Shared World")
+	int32 PreviewImageBytesTotal;
+	
+	/** The size of the AR world data that will be replicated to each client */
+	UPROPERTY(BlueprintReadOnly, Category="AR Shared World")
+	int32 ARWorldBytesTotal;
+	
+	/** The amount of the preview image data that has been replicated to this client so far */
+	UPROPERTY(BlueprintReadOnly, Category="AR Shared World")
+	int32 PreviewImageBytesDelivered;
+	
+	/** The amount of the AR world data that has been replicated to this client so far */
+	UPROPERTY(BlueprintReadOnly, Category="AR Shared World")
+	int32 ARWorldBytesDelivered;
+	
+	/**
+	 * Used to setup the initial values and size the arrays (client)
+	 *
+	 * @param PreviewImageSize the total size in bytes of the image data that will be sent
+	 * @param ARWorldDataSize the total size in bytes of the AR world data that will be sent
+	 */
+	void InitSharedWorld(int32 PreviewImageSize, int32 ARWorldDataSize);
+	
+	/**
+	 * Copies the buffer into the image data (client)
+	 *
+	 * @param Offset where in the buffer to start the copying
+	 * @param Buffer the chunk of data to copy
+	 * @param BufferSize the amount of data to copy
+	 */
+	void UpdatePreviewImageData(int32 Offset, const uint8* Buffer, int32 Size);
+
+	/**
+	 * Copies the buffer into the AR world data (client)
+	 *
+	 * @param Offset where in the buffer to start the copying
+	 * @param Buffer the chunk of data to copy
+	 * @param BufferSize the amount of data to copy
+	 */
+	void UpdateARWorldData(int32 Offset, const uint8* Buffer, int32 Size);
+
+	UFUNCTION(BlueprintImplementableEvent, Category="AR Shared World", meta = (DisplayName = "OnARWorldMapIsReady", ScriptName = "OnARWorldMapIsReady"))
+	void K2_OnARWorldMapIsReady();
+
+private:
+	void TriggerCompletionIfDone();
+	
+	bool bFiredCompletionEvent;
+};
diff --git a/Engine/Source/Runtime/AugmentedReality/Public/ARSharedWorldPlayerController.h b/Engine/Source/Runtime/AugmentedReality/Public/ARSharedWorldPlayerController.h
new file mode 100644
--- /dev/null
+++ b/Engine/Source/Runtime/AugmentedReality/Public/ARSharedWorldPlayerController.h
@@ -0,0 +1,62 @@
+// Copyright 1998-2018 Epic Games, Inc. All Rights Reserved.
+
+#pragma once
+
+#include "GameFramework/PlayerController.h"
+#include "ARSharedWorldPlayerController.generated.h"
+
+class AARSharedWorldGameState;
+
+UCLASS()
+class AUGMENTEDREALITY_API AARSharedWorldPlayerController :
+	public APlayerController
+{
+	GENERATED_UCLASS_BODY()
+
+public:
+	/**
+	 * Tells the server it is ready for receiving any shared world data
+	 */
+	UFUNCTION(Reliable, Server, WithValidation)
+	void ServerMarkReadyForReceiving();
+	
+	/**
+	 * Used to setup the initial values and size the arrays (client)
+	 *
+	 * @param PreviewImageSize the total size in bytes of the image data that will be sent
+	 * @param ARWorldDataSize the total size in bytes of the AR world data that will be sent
+	 */
+	UFUNCTION(Reliable, Client, WithValidation)
+	void ClientInitSharedWorld(int32 PreviewImageSize, int32 ARWorldDataSize);
+
+	/**
+	 * Copies the buffer into the image data (client)
+	 *
+	 * @param Offset where in the buffer to start the copying
+	 * @param Buffer the chunk of data to copy
+	 * @param BufferSize the amount of data to copy
+	 */
+	UFUNCTION(Reliable, Client, WithValidation)
+	void ClientUpdatePreviewImageData(int32 Offset, const TArray<uint8>& Buffer);
+	
+	/**
+	 * Copies the buffer into the AR world data (client)
+	 *
+	 * @param Offset where in the buffer to start the copying
+	 * @param Buffer the chunk of data to copy
+	 * @param BufferSize the amount of data to copy
+	 */
+	UFUNCTION(Reliable, Client, WithValidation)
+	void ClientUpdateARWorldData(int32 Offset, const TArray<uint8>& Buffer);
+	
+	/** @return whether this player can start receiving the ar world data */
+	bool IsReadyToReceive() const { return bIsReadyToReceive; }
+
+private:
+	virtual void Tick(float DeltaSeconds) override;
+
+	AARSharedWorldGameState* GetGameState();
+	bool IsGameStateReady();
+
+	bool bIsReadyToReceive;
+};
