Added performance tuning options for Face AR to ARSessionConfig
Added whether to mirror or be face relative for Face AR to ARSessionConfig

diff --git a/Engine/Plugins/Runtime/AppleARKit/Source/AppleARKit/Private/AppleARKitSystem.cpp b/Engine/Plugins/Runtime/AppleARKit/Source/AppleARKit/Private/AppleARKitSystem.cpp
--- a/Engine/Plugins/Runtime/AppleARKit/Source/AppleARKit/Private/AppleARKitSystem.cpp
+++ b/Engine/Plugins/Runtime/AppleARKit/Source/AppleARKit/Private/AppleARKitSystem.cpp
@@ -1349,7 +1349,8 @@ void FAppleARKitSystem::SessionDidAddAnchors_DelegateThread( NSArray<ARAnchor*>*
 	if (FaceARSupport != nullptr)
 	{
 		const FRotator& AdjustBy = GetSessionConfig().GetWorldAlignment() == EARWorldAlignment::Camera ? DerivedTrackingToUnrealRotation : FRotator::ZeroRotator;
-		const TArray<TSharedPtr<FAppleARKitAnchorData>> AnchorList = FaceARSupport->MakeAnchorData(anchors, GameThreadTimestamp, GameThreadFrameNumber, AdjustBy);
+		const EARFaceTrackingUpdate UpdateSetting = GetSessionConfig().GetFaceTrackingUpdate();
+		const TArray<TSharedPtr<FAppleARKitAnchorData>> AnchorList = FaceARSupport->MakeAnchorData(anchors, GameThreadTimestamp, GameThreadFrameNumber, AdjustBy, UpdateSetting);
 		for (TSharedPtr<FAppleARKitAnchorData> NewAnchorData : AnchorList)
 		{
 			auto AddAnchorTask = FSimpleDelegateGraphTask::FDelegate::CreateSP(this, &FAppleARKitSystem::SessionDidAddAnchors_Internal, NewAnchorData.ToSharedRef());
@@ -1379,7 +1380,8 @@ void FAppleARKitSystem::SessionDidUpdateAnchors_DelegateThread( NSArray<ARAnchor
 	if (FaceARSupport != nullptr)
 	{
 		const FRotator& AdjustBy = GetSessionConfig().GetWorldAlignment() == EARWorldAlignment::Camera ? DerivedTrackingToUnrealRotation : FRotator::ZeroRotator;
-		const TArray<TSharedPtr<FAppleARKitAnchorData>> AnchorList = FaceARSupport->MakeAnchorData(anchors, GameThreadTimestamp, GameThreadFrameNumber, AdjustBy);
+		const EARFaceTrackingUpdate UpdateSetting = GetSessionConfig().GetFaceTrackingUpdate();
+		const TArray<TSharedPtr<FAppleARKitAnchorData>> AnchorList = FaceARSupport->MakeAnchorData(anchors, GameThreadTimestamp, GameThreadFrameNumber, AdjustBy, UpdateSetting);
 		for (TSharedPtr<FAppleARKitAnchorData> NewAnchorData : AnchorList)
 		{
 			auto UpdateAnchorTask = FSimpleDelegateGraphTask::FDelegate::CreateSP(this, &FAppleARKitSystem::SessionDidUpdateAnchors_Internal, NewAnchorData.ToSharedRef());
diff --git a/Engine/Plugins/Runtime/AppleARKit/Source/AppleARKit/Public/AppleARKitFaceSupport.h b/Engine/Plugins/Runtime/AppleARKit/Source/AppleARKit/Public/AppleARKitFaceSupport.h
--- a/Engine/Plugins/Runtime/AppleARKit/Source/AppleARKit/Public/AppleARKitFaceSupport.h
+++ b/Engine/Plugins/Runtime/AppleARKit/Source/AppleARKit/Public/AppleARKitFaceSupport.h
@@ -27,10 +27,11 @@ public:
 	 * @param Timestamp the timestamp of this update
 	 * @param FrameNumber the frame number for this update
 	 * @param AdjustBy the additional rotation to apply to put the rotation in the proper space (camera alignment only)
+	 * @param UpdateSetting whether to just update curves or geo too
 	 *
 	 * @return the set of face anchors to dispatch
 	 */
-	virtual TArray<TSharedPtr<FAppleARKitAnchorData>> MakeAnchorData(NSArray<ARAnchor*>* NewAnchors, double Timestamp, uint32 FrameNumber, const FRotator& AdjustBy) { return TArray<TSharedPtr<FAppleARKitAnchorData>>(); }
+	virtual TArray<TSharedPtr<FAppleARKitAnchorData>> MakeAnchorData(NSArray<ARAnchor*>* NewAnchors, double Timestamp, uint32 FrameNumber, const FRotator& AdjustBy, EARFaceTrackingUpdate UpdateSetting) { return TArray<TSharedPtr<FAppleARKitAnchorData>>(); }
 
 	/**
 	 * Publishes any face AR data that needs to be sent to LiveLink. Done as a separate step because MakeAnchorData is called
diff --git a/Engine/Plugins/Runtime/AppleARKitFaceSupport/Source/AppleARKitFaceSupport/Private/AppleARKitFaceSupportImpl.cpp b/Engine/Plugins/Runtime/AppleARKitFaceSupport/Source/AppleARKitFaceSupport/Private/AppleARKitFaceSupportImpl.cpp
--- a/Engine/Plugins/Runtime/AppleARKitFaceSupport/Source/AppleARKitFaceSupport/Private/AppleARKitFaceSupportImpl.cpp
+++ b/Engine/Plugins/Runtime/AppleARKitFaceSupport/Source/AppleARKitFaceSupport/Private/AppleARKitFaceSupportImpl.cpp
@@ -13,7 +13,7 @@ DECLARE_CYCLE_STAT(TEXT("Conversion"), STAT_FaceAR_Conversion, STATGROUP_FaceAR)
 
 #if SUPPORTS_ARKIT_1_0
 
-static TSharedPtr<FAppleARKitAnchorData> MakeAnchorData(ARAnchor* Anchor, const FRotator& AdjustBy)
+static TSharedPtr<FAppleARKitAnchorData> MakeAnchorData(ARAnchor* Anchor, const FRotator& AdjustBy, EARFaceTrackingUpdate UpdateSetting)
 {
 	SCOPE_CYCLE_COUNTER(STAT_FaceAR_Conversion);
 	
@@ -36,13 +36,13 @@ static TSharedPtr<FAppleARKitAnchorData> MakeAnchorData(ARAnchor* Anchor, const
 			FAppleARKitConversion::ToFGuid(FaceAnchor.identifier),
 			FAppleARKitConversion::ToFTransform(FaceAnchor.transform, AdjustBy),
 			ToBlendShapeMap(FaceAnchor.blendShapes, FAppleARKitConversion::ToFTransform(FaceAnchor.transform, AdjustBy), LeftEyeTransform, RightEyeTransform),
-			ToVertexBuffer(FaceAnchor.geometry.vertices, FaceAnchor.geometry.vertexCount),
+			UpdateSetting == EARFaceTrackingUpdate::CurvesAndGeo ? ToVertexBuffer(FaceAnchor.geometry.vertices, FaceAnchor.geometry.vertexCount) : TArray<FVector>(),
 			LeftEyeTransform,
 			RightEyeTransform,
 			LookAtTarget
 		);
         // Only convert from 16bit to 32bit once
-        if (FAppleARKitAnchorData::FaceIndices.Num() == 0)
+        if (UpdateSetting == EARFaceTrackingUpdate::CurvesAndGeo && FAppleARKitAnchorData::FaceIndices.Num() == 0)
         {
             FAppleARKitAnchorData::FaceIndices = To32BitIndexBuffer(FaceAnchor.geometry.triangleIndices, FaceAnchor.geometry.triangleCount * 3);
         }
@@ -106,13 +106,13 @@ ARConfiguration* FAppleARKitFaceSupport::ToARConfiguration(UARSessionConfig* Ses
 	return SessionConfiguration;
 }
 
-TArray<TSharedPtr<FAppleARKitAnchorData>> FAppleARKitFaceSupport::MakeAnchorData(NSArray<ARAnchor*>* Anchors, double Timestamp, uint32 FrameNumber, const FRotator& AdjustBy)
+TArray<TSharedPtr<FAppleARKitAnchorData>> FAppleARKitFaceSupport::MakeAnchorData(NSArray<ARAnchor*>* Anchors, double Timestamp, uint32 FrameNumber, const FRotator& AdjustBy, EARFaceTrackingUpdate UpdateSetting)
 {
 	TArray<TSharedPtr<FAppleARKitAnchorData>> AnchorList;
 
 	for (ARAnchor* Anchor in Anchors)
 	{
-		TSharedPtr<FAppleARKitAnchorData> AnchorData = ::MakeAnchorData(Anchor, AdjustBy);
+		TSharedPtr<FAppleARKitAnchorData> AnchorData = ::MakeAnchorData(Anchor, AdjustBy, UpdateSetting);
 		if (AnchorData.IsValid())
 		{
 			AnchorList.Add(AnchorData);
diff --git a/Engine/Plugins/Runtime/AppleARKitFaceSupport/Source/AppleARKitFaceSupport/Private/AppleARKitFaceSupportImpl.h b/Engine/Plugins/Runtime/AppleARKitFaceSupport/Source/AppleARKitFaceSupport/Private/AppleARKitFaceSupportImpl.h
--- a/Engine/Plugins/Runtime/AppleARKitFaceSupport/Source/AppleARKitFaceSupport/Private/AppleARKitFaceSupportImpl.h
+++ b/Engine/Plugins/Runtime/AppleARKitFaceSupport/Source/AppleARKitFaceSupport/Private/AppleARKitFaceSupportImpl.h
@@ -22,7 +22,7 @@ public:
 private:
 #if SUPPORTS_ARKIT_1_0
 	// ~IAppleARKitFaceSupport
-	virtual TArray<TSharedPtr<FAppleARKitAnchorData>> MakeAnchorData(NSArray<ARAnchor*>* NewAnchors, double Timestamp, uint32 FrameNumber, const FRotator& AdjustBy) override;
+	virtual TArray<TSharedPtr<FAppleARKitAnchorData>> MakeAnchorData(NSArray<ARAnchor*>* NewAnchors, double Timestamp, uint32 FrameNumber, const FRotator& AdjustBy, EARFaceTrackingUpdate UpdateSetting) override;
 	virtual ARConfiguration* ToARConfiguration(UARSessionConfig* InSessionConfig) override;
 	virtual void PublishLiveLinkData(TSharedPtr<FAppleARKitAnchorData> Anchor, double Timestamp, uint32 FrameNumber) override;
 	virtual bool DoesSupportFaceAR() override;
diff --git a/Engine/Source/Runtime/AugmentedReality/Private/ARSessionConfig.cpp b/Engine/Source/Runtime/AugmentedReality/Private/ARSessionConfig.cpp
--- a/Engine/Source/Runtime/AugmentedReality/Private/ARSessionConfig.cpp
+++ b/Engine/Source/Runtime/AugmentedReality/Private/ARSessionConfig.cpp
@@ -70,7 +70,6 @@ int32 UARSessionConfig::GetMaxNumSimultaneousImagesTracked() const
     return MaxNumSimultaneousImagesTracked;
 }
 
-/** @see EnvironmentCaptureProbeType */
 EAREnvironmentCaptureProbeType UARSessionConfig::GetEnvironmentCaptureProbeType() const
 {
 	return EnvironmentCaptureProbeType;
@@ -91,6 +90,11 @@ const TArray<UARCandidateObject*>& UARSessionConfig::GetCandidateObjectList() co
 	return CandidateObjects;
 }
 
+void UARSessionConfig::SetCandidateObjectList(const TArray<UARCandidateObject*>& InCandidateObjects)
+{
+	CandidateObjects = InCandidateObjects;
+}
+
 void UARSessionConfig::AddCandidateObject(UARCandidateObject* CandidateObject)
 {
 	if (CandidateObject != nullptr)
@@ -109,6 +113,26 @@ void UARSessionConfig::SetDesiredVideoFormat(FARVideoFormat NewFormat)
 	DesiredVideoFormat = NewFormat;
 }
 
+EARFaceTrackingDirection UARSessionConfig::GetFaceTrackingDirection() const
+{
+	return FaceTrackingDirection;
+}
+
+void UARSessionConfig::SetFaceTrackingDirection(EARFaceTrackingDirection InDirection)
+{
+	FaceTrackingDirection = InDirection;
+}
+
+EARFaceTrackingUpdate UARSessionConfig::GetFaceTrackingUpdate() const
+{
+	return FaceTrackingUpdate;
+}
+
+void UARSessionConfig::SetFaceTrackingUpdate(EARFaceTrackingUpdate InUpdate)
+{
+	FaceTrackingUpdate = InUpdate;
+}
+
 void UARSessionConfig::Serialize(FArchive& Ar)
 {
 	Ar.UsingCustomVersion(FVRObjectVersion::GUID);
diff --git a/Engine/Source/Runtime/AugmentedReality/Public/ARSessionConfig.h b/Engine/Source/Runtime/AugmentedReality/Public/ARSessionConfig.h
--- a/Engine/Source/Runtime/AugmentedReality/Public/ARSessionConfig.h
+++ b/Engine/Source/Runtime/AugmentedReality/Public/ARSessionConfig.h
@@ -97,7 +97,19 @@ enum class EAREnvironmentCaptureProbeType : uint8
 	Automatic
 };
 
-UCLASS(BlueprintType, Category="AR AugmentedReality")
+/**
+ * Tells the AR system how much of the face work to perform
+ */
+UENUM(BlueprintType)
+enum class EARFaceTrackingUpdate : uint8
+{
+	/** Curves and geometry will be updated (only needed for mesh visualization) */
+	CurvesAndGeo,
+	/** Only the curve data is updated */
+	CurvesOnly
+};
+
+UCLASS(BlueprintType, Category="AR Settings")
 class AUGMENTEDREALITY_API UARSessionConfig : public UDataAsset
 {
 	GENERATED_BODY()
@@ -108,50 +120,87 @@ public:
 	
 public:
 	/** @see EARWorldAlignment */
+	UFUNCTION(BlueprintCallable, Category = "AR Settings")
 	EARWorldAlignment GetWorldAlignment() const;
 
 	/** @see SessionType */
+	UFUNCTION(BlueprintCallable, Category = "AR Settings")
 	EARSessionType GetSessionType() const;
 
 	/** @see PlaneDetectionMode */
+	UFUNCTION(BlueprintCallable, Category = "AR Settings")
 	EARPlaneDetectionMode GetPlaneDetectionMode() const;
 
 	/** @see LightEstimationMode */
+	UFUNCTION(BlueprintCallable, Category = "AR Settings")
 	EARLightEstimationMode GetLightEstimationMode() const;
 
 	/** @see FrameSyncMode */
+	UFUNCTION(BlueprintCallable, Category = "AR Settings")
 	EARFrameSyncMode GetFrameSyncMode() const;
 
 	/** @see bEnableAutomaticCameraOverlay */
+	UFUNCTION(BlueprintCallable, Category = "AR Settings")
 	bool ShouldRenderCameraOverlay() const;
 
 	/** @see bEnableAutomaticCameraTracking */
+	UFUNCTION(BlueprintCallable, Category = "AR Settings")
 	bool ShouldEnableCameraTracking() const;
 
 	/** @see bEnableAutoFocus */
+	UFUNCTION(BlueprintCallable, Category = "AR Settings")
 	bool ShouldEnableAutoFocus() const;
 
 	/** @see CandidateImages */
+	UFUNCTION(BlueprintCallable, Category = "AR Settings")
 	const TArray<UARCandidateImage*>& GetCandidateImageList() const;
     
 	/** @see MaxNumSimultaneousImagesTracked */
+	UFUNCTION(BlueprintCallable, Category = "AR Settings")
     int32 GetMaxNumSimultaneousImagesTracked() const;
 	
 	/** @see EnvironmentCaptureProbeType */
+	UFUNCTION(BlueprintCallable, Category = "AR Settings")
 	EAREnvironmentCaptureProbeType GetEnvironmentCaptureProbeType() const;
 	
 	/** @see WorldMapData */
+	UFUNCTION(BlueprintCallable, Category = "AR Settings")
 	const TArray<uint8>& GetWorldMapData() const;
 	/** @see WorldMapData */
+	UFUNCTION(BlueprintCallable, Category = "AR Settings")
 	void SetWorldMapData(TArray<uint8> WorldMapData);
 
+	/** @see CandidateObjects */
+	UFUNCTION(BlueprintCallable, Category = "AR Settings")
 	const TArray<UARCandidateObject*>& GetCandidateObjectList() const;
+	/** @see CandidateObjects */
+	UFUNCTION(BlueprintCallable, Category = "AR Settings")
+	void SetCandidateObjectList(const TArray<UARCandidateObject*>& InCandidateObjects);
+	/** @see CandidateObjects */
+	UFUNCTION(BlueprintCallable, Category = "AR Settings")
 	void AddCandidateObject(UARCandidateObject* CandidateObject);
 	
 	/** @see DesiredVideoFormat */
+	UFUNCTION(BlueprintCallable, Category = "AR Settings")
 	FARVideoFormat GetDesiredVideoFormat() const;
+	/** @see DesiredVideoFormat */
+	UFUNCTION(BlueprintCallable, Category = "AR Settings")
 	void SetDesiredVideoFormat(FARVideoFormat NewFormat);
 
+	/** @see FaceTrackingDirection */
+	UFUNCTION(BlueprintCallable, Category = "AR Settings")
+	EARFaceTrackingDirection GetFaceTrackingDirection() const;
+	/** @see FaceTrackingDirection */
+	UFUNCTION(BlueprintCallable, Category = "AR Settings")
+	void SetFaceTrackingDirection(EARFaceTrackingDirection InDirection);
+	
+	/** @see FaceTrackingUpdate */
+	UFUNCTION(BlueprintCallable, Category = "AR Settings")
+	EARFaceTrackingUpdate GetFaceTrackingUpdate() const;
+	/** @see FaceTrackingUpdate */
+	UFUNCTION(BlueprintCallable, Category = "AR Settings")
+	void SetFaceTrackingUpdate(EARFaceTrackingUpdate InUpdate);
+	
 private:
 	//~ UObject interface
 	virtual void Serialize(FArchive& Ar) override;
@@ -224,4 +273,12 @@ protected:
 	 */
 	UPROPERTY(EditAnywhere, Category="AR Settings")
 	FARVideoFormat DesiredVideoFormat;
+	
+	/** Whether to track the face as if you are looking out of the device or as a mirror */
+	UPROPERTY(EditAnywhere, Category="Face AR Settings")
+	EARFaceTrackingDirection FaceTrackingDirection;
+
+	/** Whether to track the face as if you are looking out of the device or as a mirror */
+	UPROPERTY(EditAnywhere, Category="Face AR Settings")
+	EARFaceTrackingUpdate FaceTrackingUpdate;
 };
