Changed the AR World Save and AR Get Candidate Object latent actions to use the new mechanism to reduce code

diff --git a/Engine/Source/Runtime/AugmentedReality/Private/ARBlueprintProxy.cpp b/Engine/Source/Runtime/AugmentedReality/Private/ARBlueprintProxy.cpp
--- a/Engine/Source/Runtime/AugmentedReality/Private/ARBlueprintProxy.cpp
+++ b/Engine/Source/Runtime/AugmentedReality/Private/ARBlueprintProxy.cpp
@@ -46,74 +46,67 @@ const TSharedPtr<FARSystemBase, ESPMode::ThreadSafe>& UARBaseAsyncTaskBlueprintP
 	return RegisteredARSystem;
 }
 
-UARSaveWorldAsyncTaskBlueprintProxy::UARSaveWorldAsyncTaskBlueprintProxy(const FObjectInitializer& ObjectInitializer)
-	: Super(ObjectInitializer)
+UARSaveWorldAsyncTaskBlueprintProxy* UARSaveWorldAsyncTaskBlueprintProxy::ARSaveWorld(UObject* WorldContextObject)
 {
+	UARSaveWorldAsyncTaskBlueprintProxy* Proxy = NewObject<UARSaveWorldAsyncTaskBlueprintProxy>();
+	Proxy->RegisterWithGameInstance(WorldContextObject);
+	return Proxy;
 }
 
-UARSaveWorldAsyncTaskBlueprintProxy* UARSaveWorldAsyncTaskBlueprintProxy::CreateProxyObjectForARSaveWorld(bool bCompressData)
+void UARSaveWorldAsyncTaskBlueprintProxy::Activate()
 {
-	UARSaveWorldAsyncTaskBlueprintProxy* Proxy = NewObject<UARSaveWorldAsyncTaskBlueprintProxy>();
-	Proxy->SetFlags(RF_StrongRefOnFrame);
-
 	auto ARSystem = GetARSystem();
 	if (ARSystem.IsValid())
 	{
-		Proxy->SaveWorldTask = ARSystem->SaveWorld();
-		Proxy->AsyncTask = Proxy->SaveWorldTask;
+		SaveWorldTask = ARSystem->SaveWorld();
+		AsyncTask = SaveWorldTask;
 	}
 	else
 	{
-		Proxy->AsyncTask = MakeShared<FARErrorSaveWorldAsyncTask, ESPMode::ThreadSafe>(TEXT("SaveWorld - requires a valid, running ARKit 2.0 session"));
+		AsyncTask = MakeShared<FARErrorSaveWorldAsyncTask, ESPMode::ThreadSafe>(TEXT("ARSaveWorld - requires a valid, running session"));
 	}
-
-	return Proxy;
 }
 
 void UARSaveWorldAsyncTaskBlueprintProxy::ReportSuccess()
 {
-	SaveWorldResult.WorldData = SaveWorldTask->GetSavedWorldData();
-	OnSuccess.Broadcast(SaveWorldResult);
+	OnSuccess.Broadcast(SaveWorldTask->GetSavedWorldData());
 }
 
 void UARSaveWorldAsyncTaskBlueprintProxy::ReportFailure()
 {
-	SaveWorldResult.Error = AsyncTask->GetErrorString();
-	OnFailure.Broadcast(SaveWorldResult);
+	OnFailed.Broadcast(TArray<uint8>());
 }
 
-UARGetCandidateObjectAsyncTaskBlueprintProxy::UARGetCandidateObjectAsyncTaskBlueprintProxy(const FObjectInitializer& ObjectInitializer)
-	: Super(ObjectInitializer)
+UARGetCandidateObjectAsyncTaskBlueprintProxy* UARGetCandidateObjectAsyncTaskBlueprintProxy::ARGetCandidateObject(UObject* WorldContextObject, FVector Location, FVector Extent)
 {
+	UARGetCandidateObjectAsyncTaskBlueprintProxy* Proxy = NewObject<UARGetCandidateObjectAsyncTaskBlueprintProxy>();
+	Proxy->RegisterWithGameInstance(WorldContextObject);
+	Proxy->Extent = Extent;
+	Proxy->Location = Location;
+
+	return Proxy;
 }
 
-UARGetCandidateObjectAsyncTaskBlueprintProxy* UARGetCandidateObjectAsyncTaskBlueprintProxy::CreateProxyObjectForARGetCandidateObject(FVector Location, FVector Extent)
+void UARGetCandidateObjectAsyncTaskBlueprintProxy::Activate()
 {
-	UARGetCandidateObjectAsyncTaskBlueprintProxy* Proxy = NewObject<UARGetCandidateObjectAsyncTaskBlueprintProxy>();
-	Proxy->SetFlags(RF_StrongRefOnFrame);
-
 	auto ARSystem = GetARSystem();
 	if (ARSystem.IsValid())
 	{
-		Proxy->CandidateObjectTask = ARSystem->GetCandidateObject(Location, Extent);
-		Proxy->AsyncTask = Proxy->CandidateObjectTask;
+		CandidateObjectTask = ARSystem->GetCandidateObject(Location, Extent);
+		AsyncTask = CandidateObjectTask;
 	}
 	else
 	{
-		Proxy->AsyncTask = MakeShared<FARErrorGetCandidateObjectAsyncTask, ESPMode::ThreadSafe>(TEXT("GetCandidateObject - requires a valid, running ARKit 2.0 session"));
+		AsyncTask = MakeShared<FARErrorGetCandidateObjectAsyncTask, ESPMode::ThreadSafe>(TEXT("ARGetCandidateObject - requires a valid, running session"));
 	}
-
-	return Proxy;
 }
 
 void UARGetCandidateObjectAsyncTaskBlueprintProxy::ReportSuccess()
 {
-	CandidateObjectResult.CandidateObject = CandidateObjectTask->GetCandidateObject();
-	OnSuccess.Broadcast(CandidateObjectResult);
+	OnSuccess.Broadcast(CandidateObjectTask->GetCandidateObject());
 }
 
 void UARGetCandidateObjectAsyncTaskBlueprintProxy::ReportFailure()
 {
-	CandidateObjectResult.Error = AsyncTask->GetErrorString();
-	OnFailure.Broadcast(CandidateObjectResult);
+	OnFailed.Broadcast(nullptr);
 }
diff --git a/Engine/Source/Runtime/AugmentedReality/Public/ARBlueprintProxy.h b/Engine/Source/Runtime/AugmentedReality/Public/ARBlueprintProxy.h
--- a/Engine/Source/Runtime/AugmentedReality/Public/ARBlueprintProxy.h
+++ b/Engine/Source/Runtime/AugmentedReality/Public/ARBlueprintProxy.h
@@ -7,14 +7,15 @@
 #include "UObject/Object.h"
 #include "Engine/EngineTypes.h"
 #include "Tickable.h"
+#include "Kismet/BlueprintAsyncActionBase.h"
 
 #include "ARTypes.h"
 
 #include "ARBlueprintProxy.generated.h"
 
-UCLASS(MinimalAPI)
+UCLASS(Abstract)
 class UARBaseAsyncTaskBlueprintProxy :
-	public UObject,
+	public UBlueprintAsyncActionBase,
 	public FTickableGameObject
 {
 	GENERATED_UCLASS_BODY()
@@ -26,9 +27,6 @@ public:
 	virtual TStatId GetStatId() const override { RETURN_QUICK_DECLARE_CYCLE_STAT(UARBaseAsyncTaskBlueprintProxy, STATGROUP_Tickables); }
 	//~ End FTickableObject Interface
 	
-	/** The async task to check during Tick() */
-	TSharedPtr<FARAsyncTask, ESPMode::ThreadSafe> AsyncTask;
-	
 	virtual void ReportSuccess() { check(0); }
 	virtual void ReportFailure() { check(0); }
 
@@ -36,7 +34,9 @@ public:
 
 protected:
 	static const TSharedPtr<FARSystemBase, ESPMode::ThreadSafe>& GetARSystem();
-
+	/** The async task to check during Tick() */
+	TSharedPtr<FARAsyncTask, ESPMode::ThreadSafe> AsyncTask;
+	
 private:
 	/** True until the async task completes, then false */
 	bool bShouldTick;
@@ -44,93 +44,76 @@ private:
 	static TSharedPtr<FARSystemBase, ESPMode::ThreadSafe> RegisteredARSystem;
 };
 
-USTRUCT(BlueprintType)
-struct FARSaveWorldResult
-{
-	GENERATED_USTRUCT_BODY()
-	
-	UPROPERTY(BlueprintReadOnly, Category="Augmented Reality")
-	FString Error;
-	
-	UPROPERTY(BlueprintReadOnly, Category="Augmented Reality")
-	TArray<uint8> WorldData;
-};
-
-DECLARE_DYNAMIC_MULTICAST_DELEGATE_OneParam(FARSaveWorldDelegate, const FARSaveWorldResult&, SaveWorldResult);
+DECLARE_DYNAMIC_MULTICAST_DELEGATE_OneParam(FARSaveWorldPin, const TArray<uint8>&, SavedWorld);
 
-UCLASS(MinimalAPI)
+UCLASS()
 class UARSaveWorldAsyncTaskBlueprintProxy :
 	public UARBaseAsyncTaskBlueprintProxy
 {
-	GENERATED_UCLASS_BODY()
-	
+	GENERATED_BODY()
+
 public:
 	UPROPERTY(BlueprintAssignable)
-	FARSaveWorldDelegate OnSuccess;
+	FARSaveWorldPin OnSuccess;
 	
 	UPROPERTY(BlueprintAssignable)
-	FARSaveWorldDelegate OnFailure;
+	FARSaveWorldPin OnFailed;
 	
 	/**
 	 * Saves an AR world to a byte array for network replication or saving to disk
 	 *
 	 * @param bCompressData whether to compress the data or not
 	 */
-	UFUNCTION(BlueprintCallable, Meta=(DisplayName="AR Save World"), Category="Augmented Reality")
-	static UARSaveWorldAsyncTaskBlueprintProxy* CreateProxyObjectForARSaveWorld(bool bCompressData);
+	UFUNCTION(BlueprintCallable, meta=(BlueprintInternalUseOnly="true", Category = "Augmented Reality", WorldContext = "WorldContextObject"))
+	static UARSaveWorldAsyncTaskBlueprintProxy* ARSaveWorld(UObject* WorldContextObject);
+
+private:
+	// UBlueprintAsyncActionBase interface
+	virtual void Activate() override;
+	//~UBlueprintAsyncActionBase interface
 	
 	/** The async task to check during Tick() */
 	TSharedPtr<FARSaveWorldAsyncTask, ESPMode::ThreadSafe> SaveWorldTask;
 	
-	UPROPERTY(BlueprintReadOnly, Category="Augmented Reality")
-	FARSaveWorldResult SaveWorldResult;
-
 	virtual void ReportSuccess() override;
 	virtual void ReportFailure() override;
 };
 
-USTRUCT(BlueprintType)
-struct FARGetCandidateObjectResult
-{
-	GENERATED_USTRUCT_BODY()
-	
-	UPROPERTY(BlueprintReadOnly, Category="Augmented Reality")
-	FString Error;
-	
-	UPROPERTY(BlueprintReadOnly, Category="Augmented Reality")
-	UARCandidateObject* CandidateObject;
-};
-
-DECLARE_DYNAMIC_MULTICAST_DELEGATE_OneParam(FARGetCandidateObjectDelegate, const FARGetCandidateObjectResult&, CandidateObjectResult);
+DECLARE_DYNAMIC_MULTICAST_DELEGATE_OneParam(FARGetCandidateObjectPin, UARCandidateObject*, SavedObject);
 
-UCLASS(MinimalAPI)
+UCLASS()
 class UARGetCandidateObjectAsyncTaskBlueprintProxy :
 	public UARBaseAsyncTaskBlueprintProxy
 {
-	GENERATED_UCLASS_BODY()
-	
+	GENERATED_BODY()
+
 public:
 	UPROPERTY(BlueprintAssignable)
-	FARGetCandidateObjectDelegate OnSuccess;
+	FARGetCandidateObjectPin OnSuccess;
 	
 	UPROPERTY(BlueprintAssignable)
-	FARGetCandidateObjectDelegate OnFailure;
-	
+	FARGetCandidateObjectPin OnFailed;
+
 	/**
 	 * Saves the point cloud centered at the specified location capturing all of the features within the specified extent as an object that can be detected later
 	 *
 	 * @param Location the center of the extent to grab features at
 	 * @param Extent the size of the region to grab feature points
 	 */
-	UFUNCTION(BlueprintCallable, Meta=(DisplayName="AR Get Candidate Object"), Category="Augmented Reality")
-	static UARGetCandidateObjectAsyncTaskBlueprintProxy* CreateProxyObjectForARGetCandidateObject(FVector Location, FVector Extent);
+	UFUNCTION(BlueprintCallable, Meta=(BlueprintInternalUseOnly = "true", Category = "Augmented Reality", WorldContext = "WorldContextObject"))
+	static UARGetCandidateObjectAsyncTaskBlueprintProxy* ARGetCandidateObject(UObject* WorldContextObject, FVector Location, FVector Extent);
+	
+	FVector Location;
+	FVector Extent;
+	
+private:
+	// UBlueprintAsyncActionBase interface
+	virtual void Activate() override;
+	//~UBlueprintAsyncActionBase interface
 	
 	/** The async task to check during Tick() */
 	TSharedPtr<FARGetCandidateObjectAsyncTask, ESPMode::ThreadSafe> CandidateObjectTask;
 	
-	UPROPERTY(BlueprintReadOnly, Category="Augmented Reality")
-	FARGetCandidateObjectResult CandidateObjectResult;
-
 	virtual void ReportSuccess() override;
 	virtual void ReportFailure() override;
 };
