Added more descriptive display names for AR blueprint operations

diff --git a/Engine/Source/Runtime/AugmentedReality/Public/ARBlueprintLibrary.h b/Engine/Source/Runtime/AugmentedReality/Public/ARBlueprintLibrary.h
--- a/Engine/Source/Runtime/AugmentedReality/Public/ARBlueprintLibrary.h
+++ b/Engine/Source/Runtime/AugmentedReality/Public/ARBlueprintLibrary.h
@@ -24,15 +24,15 @@ public:
 	 *
 	 * @param SessionConfig    Describes the tracking method to use, what kind of geometry to detect in the world, etc.
 	 */
-	UFUNCTION(BlueprintCallable, Category = "AR AugmentedReality|Session", meta = (Keywords = "ar augmentedreality augmented reality session start run running"))
+	UFUNCTION(BlueprintCallable, Category = "AR AugmentedReality|Session", meta = (DisplayName="Start AR Session", Keywords = "ar augmentedreality augmented reality session start run running"))
 	static void StartARSession(UARSessionConfig* SessionConfig);
 
 	/** Pause a running Augmented Reality session without clearing existing state. */
-	UFUNCTION(BlueprintCallable, Category = "AR AugmentedReality|Session", meta = (Keywords = "ar augmentedreality augmented reality session stop run running"))
+	UFUNCTION(BlueprintCallable, Category = "AR AugmentedReality|Session", meta = (DisplayName="Pause AR Session", Keywords = "ar augmentedreality augmented reality session stop run running"))
 	static void PauseARSession();
 
 	/** Stop a running Augmented Reality session and clear any state. */
-	UFUNCTION(BlueprintCallable, Category = "AR AugmentedReality|Session", meta = (Keywords = "ar augmentedreality augmented reality session stop run running"))
+	UFUNCTION(BlueprintCallable, Category = "AR AugmentedReality|Session", meta = (DisplayName="Stop AR Session", Keywords = "ar augmentedreality augmented reality session stop run running"))
 	static void StopARSession();
 	
 	/**
@@ -41,11 +41,11 @@ public:
 	 *
 	 * @return The status of a current Augmented Reality session: e.g. Running or Not running for a specific reason.
 	 */
-	UFUNCTION(BlueprintPure, Category = "AR AugmentedReality|Session", meta = (Keywords = "ar augmentedreality augmented reality session start stop run running"))
+	UFUNCTION(BlueprintPure, Category = "AR AugmentedReality|Session", meta = (DisplayName="Get AR Session Status", Keywords = "ar augmentedreality augmented reality session start stop run running"))
 	static FARSessionStatus GetARSessionStatus();
 
 	/** @return the configuration that the current session was started with. */
-	UFUNCTION(BlueprintPure, Category = "AR AugmentedReality|Session", meta = (Keywords = "ar augmentedreality augmented reality session"))
+	UFUNCTION(BlueprintPure, Category = "AR AugmentedReality|Session", meta = (DisplayName="Get AR Session Config", Keywords = "ar augmentedreality augmented reality session"))
 	static UARSessionConfig* GetSessionConfig();
 	
 	/**
@@ -60,7 +60,7 @@ public:
 	 * \see PinComponent
 	 * \see PinComponentToTraceResult
 	 */
-	UFUNCTION(BlueprintCallable, Category = "AR AugmentedReality|Alignment", meta = (Keywords = "ar augmentedreality augmented reality tracking alignment"))
+	UFUNCTION(BlueprintCallable, Category = "AR AugmentedReality|Alignment", meta = (DisplayName="Set AR Alignment Transform", Keywords = "ar augmentedreality augmented reality tracking alignment"))
 	static void SetAlignmentTransform( const FTransform& InAlignmentTransform );
 	
 	
@@ -76,26 +76,26 @@ public:
 	
 	
 	/** @return how well the tracking system is performing at the moment */
-	UFUNCTION(BlueprintPure, Category = "AR AugmentedReality|Tracking", meta = (Keywords = "ar augmentedreality augmented reality tracking quality"))
+	UFUNCTION(BlueprintPure, Category = "AR AugmentedReality|Tracking", meta = (DisplayName="Get AR Tracking Quality", Keywords = "ar augmentedreality augmented reality tracking quality"))
 	static EARTrackingQuality GetTrackingQuality();
 	
 	/** @return a list of all the real-world geometry as currently seen by the Augmented Reality system */
-	UFUNCTION(BlueprintCallable, Category = "AR AugmentedReality|Tracking", meta = (Keywords = "ar augmentedreality augmented reality tracking geometry anchor"))
+	UFUNCTION(BlueprintCallable, Category = "AR AugmentedReality|Tracking", meta = (DisplayName="Get All AR Geometries", Keywords = "ar augmentedreality augmented reality tracking geometry anchor"))
 	static TArray<UARTrackedGeometry*> GetAllGeometries();
 
 	/** @return the current camera image from the Augmented Reality system */
-	UFUNCTION(BlueprintCallable, Category = "AR AugmentedReality|Camera", meta = (Keywords = "ar augmentedreality augmented reality camera image"))
+	UFUNCTION(BlueprintCallable, Category = "AR AugmentedReality|Camera", meta = (DisplayName="Get AR Camera Image", Keywords = "ar augmentedreality augmented reality camera image"))
 	static UARTextureCameraImage* GetCameraImage();
 
 	/** @return the current camera depth data from the Augmented Reality system */
-	UFUNCTION(BlueprintCallable, Category = "AR AugmentedReality|Camera", meta = (Keywords = "ar augmentedreality augmented reality camera image depth"))
+	UFUNCTION(BlueprintCallable, Category = "AR AugmentedReality|Camera", meta = (DisplayName="Get AR Camera Depth", Keywords = "ar augmentedreality augmented reality camera image depth"))
 	static UARTextureCameraDepth* GetCameraDepth();
 
 	/**
 	 * Test whether this type of session is supported by the current Augmented Reality platform.
 	 * e.g. is your device capable of doing positional tracking or orientation only?
 	 */
-	UFUNCTION(BlueprintPure, Category = "AR AugmentedReality|Tracking", meta = (Keywords = "ar augmentedreality augmented reality tracking"))
+	UFUNCTION(BlueprintPure, Category = "AR AugmentedReality|Tracking", meta = (DisplayName="Is AR Session Type Supported", Keywords = "ar augmentedreality augmented reality tracking"))
 	static bool IsSessionTypeSupported(EARSessionType SessionType);
 	
 	
@@ -159,31 +159,31 @@ public:
 	static TArray<UARPin*> GetAllPins();
 	
 	/** @return a list of all the tracked planes as currently seen by the Augmented Reality system */
-	UFUNCTION(BlueprintCallable, Category = "AR AugmentedReality|Tracking", meta = (Keywords = "ar augmentedreality augmented reality tracking geometry anchor"))
+	UFUNCTION(BlueprintCallable, Category = "AR AugmentedReality|Tracking", meta = (DisplayName="Get All AR Tracked Planes", Keywords = "ar augmentedreality augmented reality tracking geometry anchor"))
 	static TArray<UARPlaneGeometry*> GetAllTrackedPlanes();
 	
 	/** @return a list of all the tracked points as currently seen by the Augmented Reality system */
-	UFUNCTION(BlueprintCallable, Category = "AR AugmentedReality|Tracking", meta = (Keywords = "ar augmentedreality augmented reality tracking geometry anchor"))
+	UFUNCTION(BlueprintCallable, Category = "AR AugmentedReality|Tracking", meta = (KDisplayName="Get All AR Tracked Points", eywords = "ar augmentedreality augmented reality tracking geometry anchor"))
 	static TArray<UARTrackedPoint*> GetAllTrackedPoints();
 	
 	/** @return a list of all the tracked images as currently seen by the Augmented Reality system */
-	UFUNCTION(BlueprintCallable, Category = "AR AugmentedReality|Tracking", meta = (Keywords = "ar augmentedreality augmented reality tracking images anchor"))
+	UFUNCTION(BlueprintCallable, Category = "AR AugmentedReality|Tracking", meta = (DisplayName="Get All AR Tracked Images", Keywords = "ar augmentedreality augmented reality tracking images anchor"))
 	static TArray<UARTrackedImage*> GetAllTrackedImages();
 	
 	/** @return a list of all the tracked environment capture probes as currently seen by the Augmented Reality system */
-	UFUNCTION(BlueprintCallable, Category = "AR AugmentedReality|Tracking", meta = (Keywords = "ar augmentedreality augmented reality tracking anchor"))
+	UFUNCTION(BlueprintCallable, Category = "AR AugmentedReality|Tracking", meta = (DisplayName="Get All AR Tracked Environment Probes", Keywords = "ar augmentedreality augmented reality tracking anchor"))
 	static TArray<UAREnvironmentCaptureProbe*> GetAllTrackedEnvironmentCaptureProbes();
 
 	/** Adds an environment capture probe to the ar world */
-	UFUNCTION(BlueprintCallable, Category = "AR AugmentedReality|Tracking", meta = (Keywords = "ar augmentedreality augmented reality tracking anchor"))
+	UFUNCTION(BlueprintCallable, Category = "AR AugmentedReality|Tracking", meta = (DisplayName="Add AR Environment Probe", Keywords = "ar augmentedreality augmented reality tracking anchor"))
 	static bool AddManualEnvironmentCaptureProbe(FVector Location, FVector Extent);
 
 	/** @return the current world mapping status for the AR world */
-	UFUNCTION(BlueprintCallable, Category = "AR AugmentedReality|Tracking", meta = (Keywords = "ar augmentedreality augmented reality tracking anchor"))
+	UFUNCTION(BlueprintCallable, Category = "AR AugmentedReality|Tracking", meta = (DisplayName="Get AR World Mapping Status", Keywords = "ar augmentedreality augmented reality tracking anchor"))
 	static EARWorldMappingState GetWorldMappingStatus();
 	
 	/** @return The list of supported video formats for this device */
-	UFUNCTION(BlueprintCallable, Category = "AR AugmentedReality|Capabilities", meta = (Keywords = "ar augmentedreality augmented reality config video formats"))
+	UFUNCTION(BlueprintCallable, Category = "AR AugmentedReality|Capabilities", meta = (DisplayName="Get Supported AR Video Formats", Keywords = "ar augmentedreality augmented reality config video formats"))
 	static TArray<FARVideoFormat> GetSupportedVideoFormats(EARSessionType SessionType);
 
 	static TSharedPtr<FARSaveWorldAsyncTask, ESPMode::ThreadSafe> SaveWorld();
diff --git a/Engine/Source/Runtime/AugmentedReality/Public/ARBlueprintProxy.h b/Engine/Source/Runtime/AugmentedReality/Public/ARBlueprintProxy.h
--- a/Engine/Source/Runtime/AugmentedReality/Public/ARBlueprintProxy.h
+++ b/Engine/Source/Runtime/AugmentedReality/Public/ARBlueprintProxy.h
@@ -64,7 +64,7 @@ public:
 	 *
 	 * @param bCompressData whether to compress the data or not
 	 */
-	UFUNCTION(BlueprintCallable, meta=(BlueprintInternalUseOnly="true", Category = "Augmented Reality", WorldContext = "WorldContextObject"))
+	UFUNCTION(BlueprintCallable, meta=(DisplayName="Save AR World", BlueprintInternalUseOnly="true", Category = "Augmented Reality", WorldContext = "WorldContextObject"))
 	static UARSaveWorldAsyncTaskBlueprintProxy* ARSaveWorld(UObject* WorldContextObject);
 
 private:
@@ -100,7 +100,7 @@ public:
 	 * @param Location the center of the extent to grab features at
 	 * @param Extent the size of the region to grab feature points
 	 */
-	UFUNCTION(BlueprintCallable, Meta=(BlueprintInternalUseOnly = "true", Category = "Augmented Reality", WorldContext = "WorldContextObject"))
+	UFUNCTION(BlueprintCallable, Meta=(DisplayName="Get AR Candidate Object", BlueprintInternalUseOnly = "true", Category = "Augmented Reality", WorldContext = "WorldContextObject"))
 	static UARGetCandidateObjectAsyncTaskBlueprintProxy* ARGetCandidateObject(UObject* WorldContextObject, FVector Location, FVector Extent);
 	
 	FVector Location;
