Removed the FARSharedWorld from ARSharedWorldGameState to make things simpler/cleaner

diff --git a/Engine/Source/Runtime/AugmentedReality/Private/ARSharedWorldGameMode.cpp b/Engine/Source/Runtime/AugmentedReality/Private/ARSharedWorldGameMode.cpp
--- a/Engine/Source/Runtime/AugmentedReality/Private/ARSharedWorldGameMode.cpp
+++ b/Engine/Source/Runtime/AugmentedReality/Private/ARSharedWorldGameMode.cpp
@@ -14,12 +14,12 @@ AARSharedWorldGameMode::AARSharedWorldGameMode(const FObjectInitializer& ObjectI
 
 void AARSharedWorldGameMode::SetPreviewImageData(TArray<uint8> ImageData)
 {
-	GetARSharedWorldGameState()->ARSharedWorld.PreviewImageData = MoveTemp(ImageData);
+	GetARSharedWorldGameState()->PreviewImageData = MoveTemp(ImageData);
 }
 
 void AARSharedWorldGameMode::SetARSharedWorldData(TArray<uint8> ARWorldData)
 {
-	GetARSharedWorldGameState()->ARSharedWorld.ARWorldData = MoveTemp(ARWorldData);
+	GetARSharedWorldGameState()->ARWorldData = MoveTemp(ARWorldData);
 }
 
 AARSharedWorldGameState* AARSharedWorldGameMode::GetARSharedWorldGameState()
@@ -34,8 +34,8 @@ void AARSharedWorldGameMode::SetARWorldSharingIsReady()
 	{
 		bShouldSendSharedWorldData = true;
 
-		const FARSharedWorld& SharedWorld = GetARSharedWorldGameState()->ARSharedWorld;
-		UE_LOG(LogAR, Log, TEXT("Ready to share AR data with clients. AR world size is (%d) and preview image size is (%d)"), SharedWorld.ARWorldData.Num(), SharedWorld.PreviewImageData.Num());
+		const AARSharedWorldGameState* GameState = GetARSharedWorldGameState();
+		UE_LOG(LogAR, Log, TEXT("Ready to share AR data with clients. AR world size is (%d) and preview image size is (%d)"), GameState->ARWorldData.Num(), GameState->PreviewImageData.Num());
 	}
 }
 
@@ -45,8 +45,8 @@ void AARSharedWorldGameMode::Tick(float DeltaSeconds)
 	
 	if (bShouldSendSharedWorldData)
 	{
-		const FARSharedWorld& SharedWorld = GetARSharedWorldGameState()->ARSharedWorld;
-		
+		const AARSharedWorldGameState* GameState = GetARSharedWorldGameState();
+
 		// For each player, send them their next chunk of data if needed
 		for (auto It = GetWorld()->GetControllerIterator(); It; ++It)
 		{
@@ -58,14 +58,14 @@ void AARSharedWorldGameMode::Tick(float DeltaSeconds)
 					FARSharedWorldReplicationState& ReplState = PlayerToReplicationStateMap[PC];
 					// See if we need to send any image preview data
 					int32 CurrentOffset = ReplState.PreviewImageOffset;
-					if (CurrentOffset < SharedWorld.PreviewImageData.Num())
+					if (CurrentOffset < GameState->PreviewImageData.Num())
 					{
 						SendBuffer.Reset(BufferSizePerChunk);
 						
 						// Figure how much needs to be sent (could be less than the chunk size)
-						int32 BytesToSend = FMath::Min(BufferSizePerChunk, SharedWorld.PreviewImageData.Num() - CurrentOffset);
+						int32 BytesToSend = FMath::Min(BufferSizePerChunk, GameState->PreviewImageData.Num() - CurrentOffset);
 						SendBuffer.AddUninitialized(BytesToSend);
-						FMemory::Memcpy((void*)SendBuffer.GetData(), (void*)(SharedWorld.PreviewImageData.GetData() + CurrentOffset), BytesToSend);
+						FMemory::Memcpy((void*)SendBuffer.GetData(), (void*)(GameState->PreviewImageData.GetData() + CurrentOffset), BytesToSend);
 						
 						PC->ClientUpdatePreviewImageData(CurrentOffset, SendBuffer);
 						ReplState.PreviewImageOffset += BytesToSend;
@@ -74,14 +74,14 @@ void AARSharedWorldGameMode::Tick(float DeltaSeconds)
 					}
 					// See if we need to send any AR world data
 					CurrentOffset = ReplState.ARWorldOffset;
-					if (CurrentOffset < SharedWorld.ARWorldData.Num())
+					if (CurrentOffset < GameState->ARWorldData.Num())
 					{
 						SendBuffer.Reset(BufferSizePerChunk);
 						
 						// Figure how much needs to be sent (could be less than the chunk size)
-						int32 BytesToSend = FMath::Min(BufferSizePerChunk, SharedWorld.ARWorldData.Num() - CurrentOffset);
+						int32 BytesToSend = FMath::Min(BufferSizePerChunk, GameState->ARWorldData.Num() - CurrentOffset);
 						SendBuffer.AddUninitialized(BytesToSend);
-						FMemory::Memcpy((void*)SendBuffer.GetData(), (void*)(SharedWorld.ARWorldData.GetData() + CurrentOffset), BytesToSend);
+						FMemory::Memcpy((void*)SendBuffer.GetData(), (void*)(GameState->ARWorldData.GetData() + CurrentOffset), BytesToSend);
 						
 						PC->ClientUpdateARWorldData(CurrentOffset, SendBuffer);
 						ReplState.ARWorldOffset += BytesToSend;
@@ -93,8 +93,8 @@ void AARSharedWorldGameMode::Tick(float DeltaSeconds)
 				{
 					// Add if we haven't seen this player before
 					PlayerToReplicationStateMap.Add(PC, FARSharedWorldReplicationState());
-					PC->ClientInitSharedWorld(SharedWorld.PreviewImageData.Num(), SharedWorld.ARWorldData.Num());
-					UE_LOG(LogAR, Verbose, TEXT("Sent InitSharedWorld(%d, %d) to PC (%s)"), SharedWorld.PreviewImageData.Num(), SharedWorld.ARWorldData.Num(), *PC->GetName());
+					PC->ClientInitSharedWorld(GameState->PreviewImageData.Num(), GameState->ARWorldData.Num());
+					UE_LOG(LogAR, Verbose, TEXT("Sent InitSharedWorld(%d, %d) to PC (%s)"), GameState->PreviewImageData.Num(), GameState->ARWorldData.Num(), *PC->GetName());
 				}
 			}
 		}
diff --git a/Engine/Source/Runtime/AugmentedReality/Private/ARSharedWorldGameState.cpp b/Engine/Source/Runtime/AugmentedReality/Private/ARSharedWorldGameState.cpp
--- a/Engine/Source/Runtime/AugmentedReality/Private/ARSharedWorldGameState.cpp
+++ b/Engine/Source/Runtime/AugmentedReality/Private/ARSharedWorldGameState.cpp
@@ -31,12 +31,12 @@ void AARSharedWorldGameState::InitSharedWorld(int32 PreviewImageSize, int32 ARWo
 
 	if (PreviewImageSize > 0 && ARWorldDataSize > 0)
 	{
-		ARSharedWorld.PreviewImageData.Reset(PreviewImageSize);
-		ARSharedWorld.PreviewImageData.AddUninitialized(PreviewImageSize);
+		PreviewImageData.Reset(PreviewImageSize);
+		PreviewImageData.AddUninitialized(PreviewImageSize);
 		PreviewImageBytesTotal = PreviewImageSize;
 		
-		ARSharedWorld.ARWorldData.Reset(ARWorldDataSize);
-		ARSharedWorld.ARWorldData.AddUninitialized(ARWorldDataSize);
+		ARWorldData.Reset(ARWorldDataSize);
+		ARWorldData.AddUninitialized(ARWorldDataSize);
 		ARWorldBytesTotal = ARWorldDataSize;
 	}
 	else
@@ -56,7 +56,7 @@ void AARSharedWorldGameState::UpdatePreviewImageData(int32 Offset, const uint8*
 	
 	if (Offset >= 0 && (Offset + Size) <= PreviewImageBytesTotal)
 	{
-		uint8* PreviewImageBuffer = ARSharedWorld.PreviewImageData.GetData();
+		uint8* PreviewImageBuffer = PreviewImageData.GetData();
 		FMemory::Memcpy((void*)&PreviewImageBuffer[Offset], (void*)Buffer, Size);
 		PreviewImageBytesDelivered += Size;
 
@@ -80,7 +80,7 @@ void AARSharedWorldGameState::UpdateARWorldData(int32 Offset, const uint8* Buffe
 	
 	if (Offset >= 0 && (Offset + Size) <= ARWorldBytesTotal)
 	{
-		uint8* ARWorldBuffer = ARSharedWorld.ARWorldData.GetData();
+		uint8* ARWorldBuffer = ARWorldData.GetData();
 		FMemory::Memcpy((void*)&ARWorldBuffer[Offset], (void*)Buffer, Size);
 		ARWorldBytesDelivered += Size;
 		
diff --git a/Engine/Source/Runtime/AugmentedReality/Public/ARSharedWorldGameState.h b/Engine/Source/Runtime/AugmentedReality/Public/ARSharedWorldGameState.h
--- a/Engine/Source/Runtime/AugmentedReality/Public/ARSharedWorldGameState.h
+++ b/Engine/Source/Runtime/AugmentedReality/Public/ARSharedWorldGameState.h
@@ -7,17 +7,6 @@
 #include "GameFramework/GameState.h"
 #include "ARSharedWorldGameState.generated.h"
 
-USTRUCT(BlueprintType)
-struct AUGMENTEDREALITY_API FARSharedWorld
-{
-	GENERATED_BODY()
-
-	/** The image taken at the time of world saving for use when aligning the AR world later in the session */
-	TArray<uint8> PreviewImageData;
-	/** The snapshot of the AR world that is to be shared */
-	TArray<uint8> ARWorldData;
-};
-
 UCLASS(BlueprintType)
 class AUGMENTEDREALITY_API AARSharedWorldGameState :
 	public AGameState
@@ -27,7 +16,12 @@ class AUGMENTEDREALITY_API AARSharedWorldGameState :
 public:
 	/** Each client and the host have a copy of the shared world data */
 	UPROPERTY(BlueprintReadOnly, Category="AR Shared World")
-	FARSharedWorld ARSharedWorld;
+	/** The image taken at the time of world saving for use when aligning the AR world later in the session */
+	TArray<uint8> PreviewImageData;
+
+	/** Each client and the host have a copy of the shared world data */
+	UPROPERTY(BlueprintReadOnly, Category="AR Shared World")
+	TArray<uint8> ARWorldData;
 
 	/** The size of the image that will be replicated to each client */
 	UPROPERTY(BlueprintReadOnly, Category="AR Shared World")
