Fixed the shadowing warning

diff --git a/Engine/Source/Runtime/AugmentedReality/Private/ARSharedWorldGameMode.cpp b/Engine/Source/Runtime/AugmentedReality/Private/ARSharedWorldGameMode.cpp
--- a/Engine/Source/Runtime/AugmentedReality/Private/ARSharedWorldGameMode.cpp
+++ b/Engine/Source/Runtime/AugmentedReality/Private/ARSharedWorldGameMode.cpp
@@ -34,8 +34,8 @@ void AARSharedWorldGameMode::SetARWorldSharingIsReady()
 	{
 		bShouldSendSharedWorldData = true;
 
-		const AARSharedWorldGameState* GameState = GetARSharedWorldGameState();
-		UE_LOG(LogAR, Log, TEXT("Ready to share AR data with clients. AR world size is (%d) and preview image size is (%d)"), GameState->ARWorldData.Num(), GameState->PreviewImageData.Num());
+		const AARSharedWorldGameState* SharedGameState = GetARSharedWorldGameState();
+		UE_LOG(LogAR, Log, TEXT("Ready to share AR data with clients. AR world size is (%d) and preview image size is (%d)"), SharedGameState->ARWorldData.Num(), SharedGameState->PreviewImageData.Num());
 	}
 }
 
@@ -45,7 +45,7 @@ void AARSharedWorldGameMode::Tick(float DeltaSeconds)
 	
 	if (bShouldSendSharedWorldData)
 	{
-		const AARSharedWorldGameState* GameState = GetARSharedWorldGameState();
+		const AARSharedWorldGameState* SharedGameState = GetARSharedWorldGameState();
 
 		// For each player, send them their next chunk of data if needed
 		for (auto It = GetWorld()->GetControllerIterator(); It; ++It)
@@ -58,14 +58,14 @@ void AARSharedWorldGameMode::Tick(float DeltaSeconds)
 					FARSharedWorldReplicationState& ReplState = PlayerToReplicationStateMap[PC];
 					// See if we need to send any image preview data
 					int32 CurrentOffset = ReplState.PreviewImageOffset;
-					if (CurrentOffset < GameState->PreviewImageData.Num())
+					if (CurrentOffset < SharedGameState->PreviewImageData.Num())
 					{
 						SendBuffer.Reset(BufferSizePerChunk);
 						
 						// Figure how much needs to be sent (could be less than the chunk size)
-						int32 BytesToSend = FMath::Min(BufferSizePerChunk, GameState->PreviewImageData.Num() - CurrentOffset);
+						int32 BytesToSend = FMath::Min(BufferSizePerChunk, SharedGameState->PreviewImageData.Num() - CurrentOffset);
 						SendBuffer.AddUninitialized(BytesToSend);
-						FMemory::Memcpy((void*)SendBuffer.GetData(), (void*)(GameState->PreviewImageData.GetData() + CurrentOffset), BytesToSend);
+						FMemory::Memcpy((void*)SendBuffer.GetData(), (void*)(SharedGameState->PreviewImageData.GetData() + CurrentOffset), BytesToSend);
 						
 						PC->ClientUpdatePreviewImageData(CurrentOffset, SendBuffer);
 						ReplState.PreviewImageOffset += BytesToSend;
@@ -74,14 +74,14 @@ void AARSharedWorldGameMode::Tick(float DeltaSeconds)
 					}
 					// See if we need to send any AR world data
 					CurrentOffset = ReplState.ARWorldOffset;
-					if (CurrentOffset < GameState->ARWorldData.Num())
+					if (CurrentOffset < SharedGameState->ARWorldData.Num())
 					{
 						SendBuffer.Reset(BufferSizePerChunk);
 						
 						// Figure how much needs to be sent (could be less than the chunk size)
-						int32 BytesToSend = FMath::Min(BufferSizePerChunk, GameState->ARWorldData.Num() - CurrentOffset);
+						int32 BytesToSend = FMath::Min(BufferSizePerChunk, SharedGameState->ARWorldData.Num() - CurrentOffset);
 						SendBuffer.AddUninitialized(BytesToSend);
-						FMemory::Memcpy((void*)SendBuffer.GetData(), (void*)(GameState->ARWorldData.GetData() + CurrentOffset), BytesToSend);
+						FMemory::Memcpy((void*)SendBuffer.GetData(), (void*)(SharedGameState->ARWorldData.GetData() + CurrentOffset), BytesToSend);
 						
 						PC->ClientUpdateARWorldData(CurrentOffset, SendBuffer);
 						ReplState.ARWorldOffset += BytesToSend;
@@ -93,8 +93,8 @@ void AARSharedWorldGameMode::Tick(float DeltaSeconds)
 				{
 					// Add if we haven't seen this player before
 					PlayerToReplicationStateMap.Add(PC, FARSharedWorldReplicationState());
-					PC->ClientInitSharedWorld(GameState->PreviewImageData.Num(), GameState->ARWorldData.Num());
-					UE_LOG(LogAR, Verbose, TEXT("Sent InitSharedWorld(%d, %d) to PC (%s)"), GameState->PreviewImageData.Num(), GameState->ARWorldData.Num(), *PC->GetName());
+					PC->ClientInitSharedWorld(SharedGameState->PreviewImageData.Num(), SharedGameState->ARWorldData.Num());
+					UE_LOG(LogAR, Verbose, TEXT("Sent InitSharedWorld(%d, %d) to PC (%s)"), SharedGameState->PreviewImageData.Num(), SharedGameState->ARWorldData.Num(), *PC->GetName());
 				}
 			}
 		}
