Added support for telling the AR system whether to reset tracking and tracked objects (useful for generating a play space and then using lower cpu/gpu tracking only)

diff --git a/Engine/Plugins/Runtime/AppleARKit/Source/AppleARKit/Private/AppleARKitSystem.cpp b/Engine/Plugins/Runtime/AppleARKit/Source/AppleARKit/Private/AppleARKitSystem.cpp
--- a/Engine/Plugins/Runtime/AppleARKit/Source/AppleARKit/Private/AppleARKitSystem.cpp
+++ b/Engine/Plugins/Runtime/AppleARKit/Source/AppleARKit/Private/AppleARKitSystem.cpp
@@ -1120,13 +1120,6 @@ void FAppleARKitSystem::SetDeviceOrientation( EScreenOrientation::Type InOrienta
 PRAGMA_DISABLE_OPTIMIZATION
 bool FAppleARKitSystem::Run(UARSessionConfig* SessionConfig)
 {
-	if (IsRunning())
-	{
-		UE_LOG(LogAppleARKit, Log, TEXT("Session already running"), this);
-
-		return true;
-	}
-
 	{
 		// Clear out any existing frames since they aren't valid anymore
 		FScopeLock ScopeLock(&FrameLock);
@@ -1178,8 +1171,16 @@ bool FAppleARKitSystem::Run(UARSessionConfig* SessionConfig)
 		}
 		else
 		{
-			// pause and start with new options
-			options = ARSessionRunOptionResetTracking | ARSessionRunOptionRemoveExistingAnchors;
+			// Check what the user has set for reseting options
+			if (SessionConfig->ShouldResetCameraTracking())
+			{
+				options |= ARSessionRunOptionResetTracking;
+			}
+			if (SessionConfig->ShouldResetTrackedObjects())
+			{
+				options |= ARSessionRunOptionRemoveExistingAnchors;
+			}
+
 			[Session pause];
 		}
 
diff --git a/Engine/Source/Runtime/AugmentedReality/Private/ARSessionConfig.cpp b/Engine/Source/Runtime/AugmentedReality/Private/ARSessionConfig.cpp
--- a/Engine/Source/Runtime/AugmentedReality/Private/ARSessionConfig.cpp
+++ b/Engine/Source/Runtime/AugmentedReality/Private/ARSessionConfig.cpp
@@ -14,6 +14,8 @@ UARSessionConfig::UARSessionConfig()
 , FrameSyncMode(EARFrameSyncMode::SyncTickWithoutCameraImage)
 , bEnableAutomaticCameraOverlay(true)
 , bEnableAutomaticCameraTracking(true)
+, bResetCameraTracking(true)
+, bResetTrackedObjects(true)
 , MaxNumSimultaneousImagesTracked(1)
 {
 }
@@ -60,6 +62,31 @@ bool UARSessionConfig::ShouldEnableAutoFocus() const
 	return bEnableAutoFocus;
 }
 
+void UARSessionConfig::SetEnableAutoFocus(bool bNewValue)
+{
+	bEnableAutoFocus = bNewValue;
+}
+
+bool UARSessionConfig::ShouldResetCameraTracking() const
+{
+	return bResetCameraTracking;
+}
+
+void UARSessionConfig::SetResetCameraTracking(bool bNewValue)
+{
+	bResetCameraTracking = bNewValue;
+}
+
+bool UARSessionConfig::ShouldResetTrackedObjects() const
+{
+	return bResetTrackedObjects;
+}
+
+void UARSessionConfig::SetResetTrackedObjects(bool bNewValue)
+{
+	bResetTrackedObjects = bNewValue;
+}
+
 const TArray<UARCandidateImage*>& UARSessionConfig::GetCandidateImageList() const
 {
 	return CandidateImages;
diff --git a/Engine/Source/Runtime/AugmentedReality/Public/ARSessionConfig.h b/Engine/Source/Runtime/AugmentedReality/Public/ARSessionConfig.h
--- a/Engine/Source/Runtime/AugmentedReality/Public/ARSessionConfig.h
+++ b/Engine/Source/Runtime/AugmentedReality/Public/ARSessionConfig.h
@@ -151,6 +151,26 @@ public:
 	UFUNCTION(BlueprintCallable, Category = "AR Settings")
 	bool ShouldEnableAutoFocus() const;
 
+	/** @see bEnableAutoFocus */
+	UFUNCTION(BlueprintCallable, Category = "AR Settings")
+	void SetEnableAutoFocus(bool bNewValue);
+
+	/** @see bResetCameraTracking */
+	UFUNCTION(BlueprintCallable, Category = "AR Settings")
+	bool ShouldResetCameraTracking() const;
+	
+	/** @see bResetCameraTracking */
+	UFUNCTION(BlueprintCallable, Category = "AR Settings")
+	void SetResetCameraTracking(bool bNewValue);
+
+	/** @see bResetTrackedObjects */
+	UFUNCTION(BlueprintCallable, Category = "AR Settings")
+	bool ShouldResetTrackedObjects() const;
+
+	/** @see bResetTrackedObjects */
+	UFUNCTION(BlueprintCallable, Category = "AR Settings")
+	void SetResetTrackedObjects(bool bNewValue);
+	
 	/** @see CandidateImages */
 	UFUNCTION(BlueprintCallable, Category = "AR Settings")
 	const TArray<UARCandidateImage*>& GetCandidateImageList() const;
@@ -247,6 +267,14 @@ protected:
 	UPROPERTY(EditAnywhere, Category="AR Settings")
 	bool bEnableAutomaticCameraTracking;
 
+	/** Whether the AR system should reset camera tracking (origin, transform) or not. Defaults to true. */
+	UPROPERTY(EditAnywhere, Category="AR Settings")
+	bool bResetCameraTracking;
+	
+	/** Whether the AR system should remove any tracked objects or not. Defaults to true. */
+	UPROPERTY(EditAnywhere, Category="AR Settings")
+	bool bResetTrackedObjects;
+	
 	/** The list of candidate images to detect within the AR camera view */
 	UPROPERTY(EditAnywhere, Category="AR Settings")
 	TArray<UARCandidateImage*> CandidateImages;
