Header changes required for ARKit 3 support.

diff --git a/Engine/Source/Developer/Apple/MetalShaderFormat/Private/MetalShaderCompiler.cpp b/Engine/Source/Developer/Apple/MetalShaderFormat/Private/MetalShaderCompiler.cpp
--- a/Engine/Source/Developer/Apple/MetalShaderFormat/Private/MetalShaderCompiler.cpp
+++ b/Engine/Source/Developer/Apple/MetalShaderFormat/Private/MetalShaderCompiler.cpp
@@ -476,6 +476,14 @@ FString GetMetalBinaryPath(uint32 ShaderPlatform)
 		{
 			FString MetalToolsPath = bIsMobile ? FString::Printf(TEXT("%s/Platforms/iPhoneOS.platform/usr/bin"), *XcodePath) : FString::Printf(TEXT("%s/Toolchains/XcodeDefault.xctoolchain/usr/metal/macos/bin"), *XcodePath);
 			FString MetalPath = MetalToolsPath + TEXT("/metal");
+			
+			// Also consider the alternative metal path for iOS in XCode 11 beta
+			if (!RemoteFileExists(MetalPath) && bIsMobile)
+			{
+				MetalToolsPath = FString::Printf(TEXT("%s/Toolchains/XcodeDefault.xctoolchain/usr/metal/ios/bin"), *XcodePath);
+				MetalPath = MetalToolsPath + TEXT("/metal");
+			}
+			
 			if (!RemoteFileExists(MetalPath))
 			{
 				if (bIsMobile)
diff --git a/Engine/Source/Runtime/AugmentedReality/Private/ARBlueprintLibrary.cpp b/Engine/Source/Runtime/AugmentedReality/Private/ARBlueprintLibrary.cpp
--- a/Engine/Source/Runtime/AugmentedReality/Private/ARBlueprintLibrary.cpp
+++ b/Engine/Source/Runtime/AugmentedReality/Private/ARBlueprintLibrary.cpp
@@ -489,3 +489,73 @@ UARCandidateImage* UARBlueprintLibrary::AddRuntimeCandidateImage(UARSessionConfi
 	}
 }
 
+bool UARBlueprintLibrary::IsSessionTrackingFeatureSupported(EARSessionType SessionType, EARSessionTrackingFeature SessionTrackingFeature)
+{
+	auto ARSystem = GetARSystem();
+	if (ARSystem.IsValid())
+	{
+		return ARSystem.Pin()->IsSessionTrackingFeatureSupported(SessionType, SessionTrackingFeature);
+	}
+	else
+	{
+		return false;
+	}
+}
+
+TArray<FARPose2D> UARBlueprintLibrary::GetAllTracked2DPoses()
+{
+	auto ARSystem = GetARSystem();
+	if (ARSystem.IsValid())
+	{
+		return ARSystem.Pin()->GetTracked2DPose();
+	}
+	else
+	{
+		return {};
+	}
+}
+
+TArray<UARTrackedPose*> UARBlueprintLibrary::GetAllTrackedPoses()
+{
+	TArray<UARTrackedPose*> TrackedPoses;
+	
+	auto ARSystem = GetARSystem();
+	if (ARSystem.IsValid())
+	{
+		auto PinnedARSystem = ARSystem.Pin();
+		for (UARTrackedGeometry* Geo : PinnedARSystem->GetAllTrackedGeometries())
+		{
+			if (UARTrackedPose* Item = Cast<UARTrackedPose>(Geo))
+			{
+				TrackedPoses.Add(Item);
+			}
+		}
+	}
+	return TrackedPoses;
+}
+
+UARTextureCameraImage* UARBlueprintLibrary::GetPersonSegmentationImage()
+{
+	auto ARSystem = GetARSystem();
+	if (ARSystem.IsValid())
+	{
+		return ARSystem.Pin()->GetPersonSegmentationImage();
+	}
+	else
+	{
+		return nullptr;
+	}
+}
+
+UARTextureCameraImage* UARBlueprintLibrary::GetPersonSegmentationDepthImage()
+{
+	auto ARSystem = GetARSystem();
+	if (ARSystem.IsValid())
+	{
+		return ARSystem.Pin()->GetPersonSegmentationDepthImage();
+	}
+	else
+	{
+		return nullptr;
+	}
+}
diff --git a/Engine/Source/Runtime/AugmentedReality/Private/ARSessionConfig.cpp b/Engine/Source/Runtime/AugmentedReality/Private/ARSessionConfig.cpp
--- a/Engine/Source/Runtime/AugmentedReality/Private/ARSessionConfig.cpp
+++ b/Engine/Source/Runtime/AugmentedReality/Private/ARSessionConfig.cpp
@@ -200,3 +200,13 @@ void UARSessionConfig::Serialize(FArchive& Ar)
 		}
 	}
 }
+
+EARSessionTrackingFeature UARSessionConfig::GetEnabledSessionTrackingFeature() const
+{
+	return EnabledSessionTrackingFeature;
+}
+
+void UARSessionConfig::SetSessionTrackingFeatureToEnable(EARSessionTrackingFeature InSessionTrackingFeature)
+{
+	EnabledSessionTrackingFeature = InSessionTrackingFeature;
+}
diff --git a/Engine/Source/Runtime/AugmentedReality/Private/ARSupportInterface.cpp b/Engine/Source/Runtime/AugmentedReality/Private/ARSupportInterface.cpp
--- a/Engine/Source/Runtime/AugmentedReality/Private/ARSupportInterface.cpp
+++ b/Engine/Source/Runtime/AugmentedReality/Private/ARSupportInterface.cpp
@@ -323,6 +323,42 @@ void FARSupportInterface ::AddReferencedObjects(FReferenceCollector& Collector)
 	}
 }
 
+bool FARSupportInterface::IsSessionTrackingFeatureSupported(EARSessionType SessionType, EARSessionTrackingFeature SessionTrackingFeature) const
+{
+	if (ARImplemention)
+	{
+		return ARImplemention->OnIsSessionTrackingFeatureSupported(SessionType, SessionTrackingFeature);
+	}
+	return false;
+}
+
+TArray<FARPose2D> FARSupportInterface::GetTracked2DPose() const
+{
+	if (ARImplemention)
+	{
+		return ARImplemention->OnGetTracked2DPose();
+	}
+	return {};
+}
+
+UARTextureCameraImage* FARSupportInterface::GetPersonSegmentationImage() const
+{
+	if (ARImplemention)
+	{
+		return ARImplemention->OnGetPersonSegmentationImage();
+	}
+	return nullptr;
+}
+
+UARTextureCameraImage* FARSupportInterface::GetPersonSegmentationDepthImage() const
+{
+	if (ARImplemention)
+	{
+		return ARImplemention->OnGetPersonSegmentationDepthImage();
+	}
+	return nullptr;
+}
+
 #define DEFINE_AR_SI_DELEGATE_FUNCS(DelegateName) \
 FDelegateHandle FARSupportInterface::Add##DelegateName##Delegate_Handle(const F##DelegateName##Delegate& Delegate) \
 { \
diff --git a/Engine/Source/Runtime/AugmentedReality/Private/ARTrackable.cpp b/Engine/Source/Runtime/AugmentedReality/Private/ARTrackable.cpp
--- a/Engine/Source/Runtime/AugmentedReality/Private/ARTrackable.cpp
+++ b/Engine/Source/Runtime/AugmentedReality/Private/ARTrackable.cpp
@@ -332,3 +332,60 @@ void UARTrackedObject::UpdateTrackedGeometry(const TSharedRef<FARSupportInterfac
 	ObjectClassification = EARObjectClassification::SceneObject;
 }
 
+void UARTrackedPose::DebugDraw(UWorld* World, const FLinearColor& OutlineColor, float OutlineThickness, float PersistForSeconds) const
+{
+	if (TrackedPose.JointTransformSpace != EARJointTransformSpace::Model)
+	{
+		// TODO: Support joints defined in the parent space
+		return;
+	}
+	
+	const FTransform LocalToWorldTransform = GetLocalToWorldTransform();
+	const FColor OutlineRGB = OutlineColor.ToFColor(false);
+	
+	if (true)
+	{
+		// Draw the entire skeleton
+		for (int Index = 0; Index < TrackedPose.SkeletonDefinition.NumJoints; ++Index)
+		{
+			if (!TrackedPose.IsJointTracked[Index])
+			{
+				continue;
+			}
+			
+			const int ParentIndex = TrackedPose.SkeletonDefinition.ParentIndices[Index];
+			if (ParentIndex >= 0 && ParentIndex < TrackedPose.SkeletonDefinition.NumJoints)
+			{
+				if (TrackedPose.IsJointTracked[ParentIndex])
+				{
+					const FTransform JointWorldTransform = TrackedPose.JointTransforms[Index] * LocalToWorldTransform;
+					const FTransform ParentWorldTransform = TrackedPose.JointTransforms[ParentIndex] * LocalToWorldTransform;
+					DrawDebugLine(World, JointWorldTransform.GetLocation(), ParentWorldTransform.GetLocation(), OutlineRGB, false, PersistForSeconds);
+				}
+			}
+		}
+	}
+	else
+	{
+		// Draw individual joint location and name
+		for (int Index = 0; Index < TrackedPose.SkeletonDefinition.NumJoints; ++Index)
+		{
+			if (!TrackedPose.IsJointTracked[Index])
+			{
+				continue;
+			}
+			
+			const FTransform JointWorldTransform = TrackedPose.JointTransforms[Index] * LocalToWorldTransform;
+			DrawDebugPoint(World, JointWorldTransform.GetLocation(), 0.5f, OutlineRGB, false, PersistForSeconds, 0);
+			ARDebugHelpers::DrawDebugString(World, JointWorldTransform.GetLocation() + FVector(0.f, 0.f, 10.f),
+											TrackedPose.SkeletonDefinition.JointNames[Index].ToString(),
+											0.25f * OutlineThickness, OutlineRGB, PersistForSeconds, true);
+		}
+	}
+}
+
+void UARTrackedPose::UpdateTrackedPose(const TSharedRef<FARSupportInterface , ESPMode::ThreadSafe>& InTrackingSystem, uint32 FrameNumber, double Timestamp, const FTransform& InLocalToTrackingTransform, const FTransform& InAlignmentTransform, const FARPose3D& InTrackedPose)
+{
+	Super::UpdateTrackedGeometry(InTrackingSystem, FrameNumber, Timestamp, InLocalToTrackingTransform, InAlignmentTransform);
+	TrackedPose = InTrackedPose;
+}
diff --git a/Engine/Source/Runtime/AugmentedReality/Public/ARBlueprintLibrary.h b/Engine/Source/Runtime/AugmentedReality/Public/ARBlueprintLibrary.h
--- a/Engine/Source/Runtime/AugmentedReality/Public/ARBlueprintLibrary.h
+++ b/Engine/Source/Runtime/AugmentedReality/Public/ARBlueprintLibrary.h
@@ -260,6 +260,26 @@ public:
 	UFUNCTION(BlueprintCallable, Category = "AR AugmentedReality|Session", meta = (Keywords = "ar augmentedreality augmented reality candidate image"))
 	static UARCandidateImage* AddRuntimeCandidateImage(UARSessionConfig* SessionConfig, UTexture2D* CandidateTexture, FString FriendlyName, float PhysicalWidth);
 
+	/** @return if a particular session feature is supported with the specified session type on the current platform */
+	UFUNCTION(BlueprintPure, Category = "AR AugmentedReality|Session", meta = (DisplayName="Is AR Session Tracking Feature Supported", Keywords = "ar augmentedreality augmented reality session tracking feature"))
+	static bool IsSessionTrackingFeatureSupported(EARSessionType SessionType, EARSessionTrackingFeature SessionTrackingFeature);
+	
+	/** @return all the 2D poses tracked by the AR system */
+	UFUNCTION(BlueprintCallable, Category = "AR AugmentedReality|Pose Tracking", meta = (DisplayName="Get All AR Tracked 2D Poses", Keywords = "ar augmentedreality augmented reality pose tracking"))
+	static TArray<FARPose2D> GetAllTracked2DPoses();
+	
+	/** @return all the 3D poses tracked by the AR system */
+	UFUNCTION(BlueprintCallable, Category = "AR AugmentedReality|Pose Tracking", meta = (DisplayName="Get All AR Tracked 3D Poses", Keywords = "ar augmentedreality augmented reality pose tracking"))
+	static TArray<UARTrackedPose*> GetAllTrackedPoses();
+	
+	/** @return the segmentation image if the person segmentation session feature is used */
+	UFUNCTION(BlueprintCallable, Category = "AR AugmentedReality|Person Segmentation", meta = (DisplayName="Get AR Person Segmentation Image", Keywords = "ar augmentedreality augmented reality person segmentation image"))
+	static UARTextureCameraImage* GetPersonSegmentationImage();
+
+	/** @return the segmentation depth image if the person segmentation with depth session feature is used */
+	UFUNCTION(BlueprintCallable, Category = "AR AugmentedReality|Person Segmentation", meta = (DisplayName="Get AR Person Segmentation Depth Image", Keywords = "ar augmentedreality augmented reality person segmentation depth image"))
+	static UARTextureCameraImage* GetPersonSegmentationDepthImage();
+	
 	// Static helpers to create the methods needed to add/remove delegates from the AR system
 	DEFINE_AR_BPLIB_DELEGATE_FUNCS(OnTrackableAdded)
 	DEFINE_AR_BPLIB_DELEGATE_FUNCS(OnTrackableUpdated)
diff --git a/Engine/Source/Runtime/AugmentedReality/Public/ARSessionConfig.h b/Engine/Source/Runtime/AugmentedReality/Public/ARSessionConfig.h
--- a/Engine/Source/Runtime/AugmentedReality/Public/ARSessionConfig.h
+++ b/Engine/Source/Runtime/AugmentedReality/Public/ARSessionConfig.h
@@ -44,7 +44,10 @@ enum class EARSessionType : uint8
     Image,
 
 	/** A session used to scan objects for object detection in a world tracking session */
-	ObjectScanning
+	ObjectScanning,
+	
+	/** A session used to track human pose in 3D */
+	PoseTracking
 };
 
 UENUM(BlueprintType, Category = "AR AugmentedReality", meta = (Experimental, Bitflags))
@@ -109,6 +112,25 @@ enum class EARFaceTrackingUpdate : uint8
 	CurvesOnly
 };
 
+/**
+ * Tells the AR system how much of the face work to perform
+ */
+UENUM(BlueprintType)
+enum class EARSessionTrackingFeature : uint8
+{
+	/** None of the session feature is enabled */
+	None,
+	
+	/** 2D pose detection is enabled */
+	PoseDetection2D,
+	
+	/** Person segmentation is enabled */
+	PersonSegmentation,
+	
+	/** Person segmentation with depth info is enabled */
+	PersonSegmentationWithDepth,
+};
+
 UCLASS(BlueprintType, Category="AR Settings")
 class AUGMENTEDREALITY_API UARSessionConfig : public UDataAsset
 {
@@ -225,6 +247,13 @@ public:
 	UFUNCTION(BlueprintCallable, Category = "AR Settings")
 	void SetFaceTrackingUpdate(EARFaceTrackingUpdate InUpdate);
 
+	/** @see EnabledSessionTrackingFeatures */
+	UFUNCTION(BlueprintCallable, Category = "AR Settings")
+	EARSessionTrackingFeature GetEnabledSessionTrackingFeature() const;
+	
+	/** @see EnabledSessionTrackingFeatures */
+	UFUNCTION(BlueprintCallable, Category = "AR Settings")
+	void SetSessionTrackingFeatureToEnable(EARSessionTrackingFeature InSessionTrackingFeature);
 
 	bool ShouldDoHorizontalPlaneDetection() const { return bHorizontalPlaneDetection; }
 	bool ShouldDoVerticalPlaneDetection() const { return bVerticalPlaneDetection; }
@@ -255,6 +284,10 @@ public:
 	/** Whether the AR system should report scene objects (@see EARObjectClassification::SceneObject) */
 	UPROPERTY(EditAnywhere, Category = "AR Settings | World Mapping")
 	bool bTrackSceneObjects;
+	
+	/** Whether to occlude the virtual content with the result from person segmentation */
+	UPROPERTY(EditAnywhere, Category = "AR Settings | Occlusion")
+	bool bUsePersonSegmentationForOcclusion = true;
 
 private:
 	//~ UObject interface
@@ -347,5 +380,9 @@ protected:
 	
 	/** Data array for storing the cooked image database */
 	UPROPERTY()
-	TArray<uint8> SerializedARCandidateImageDatabase;	
+	TArray<uint8> SerializedARCandidateImageDatabase;
+	
+	/** A list of session features  to enable */
+	UPROPERTY(EditAnywhere, Category="AR Settings")
+	EARSessionTrackingFeature EnabledSessionTrackingFeature = EARSessionTrackingFeature::None;
 };
diff --git a/Engine/Source/Runtime/AugmentedReality/Public/ARSupportInterface.h b/Engine/Source/Runtime/AugmentedReality/Public/ARSupportInterface.h
--- a/Engine/Source/Runtime/AugmentedReality/Public/ARSupportInterface.h
+++ b/Engine/Source/Runtime/AugmentedReality/Public/ARSupportInterface.h
@@ -108,6 +108,18 @@ public:
 
 	void* GetARSessionRawPointer();
 	void* GetGameThreadARFrameRawPointer();
+	
+	/** \see UARBlueprintLibrary::IsSessionTrackingFeatureSupported() */
+	bool IsSessionTrackingFeatureSupported(EARSessionType SessionType, EARSessionTrackingFeature SessionTrackingFeature) const;
+	
+	/** \see UARBlueprintLibrary::GetTracked2DPose() */
+	TArray<FARPose2D> GetTracked2DPose() const;
+	
+	/** \see UARBlueprintLibrary::GetPersonSegmentationImage() */
+	UARTextureCameraImage* GetPersonSegmentationImage() const;
+	
+	/** \see UARBlueprintLibrary::GetPersonSegmentationDepthImage() */
+	UARTextureCameraImage* GetPersonSegmentationDepthImage() const;
 
 	//~ FGCObject
 	virtual void AddReferencedObjects(FReferenceCollector& Collector) override;
diff --git a/Engine/Source/Runtime/AugmentedReality/Public/ARSystem.h b/Engine/Source/Runtime/AugmentedReality/Public/ARSystem.h
--- a/Engine/Source/Runtime/AugmentedReality/Public/ARSystem.h
+++ b/Engine/Source/Runtime/AugmentedReality/Public/ARSystem.h
@@ -180,6 +180,18 @@ public:
 
 	virtual void* GetARSessionRawPointer() = 0;
 	virtual void* GetGameThreadARFrameRawPointer() = 0;
+	
+	/** @return if a particular session feature is supported on this device */
+	virtual bool OnIsSessionTrackingFeatureSupported(EARSessionType SessionType, EARSessionTrackingFeature SessionTrackingFeature) const { return false; }
+	
+	/** @return all the tracked 2D poses in AR */
+	virtual TArray<FARPose2D> OnGetTracked2DPose() const { return {}; }
+	
+	/** @return the person segmentation image */
+	virtual UARTextureCameraImage* OnGetPersonSegmentationImage() const { return nullptr; }
+	
+	/** @return the person segmentation depth image */
+	virtual UARTextureCameraImage* OnGetPersonSegmentationDepthImage() const { return nullptr; }
 
 	virtual ~IARSystemSupport(){}
 
diff --git a/Engine/Source/Runtime/AugmentedReality/Public/ARTrackable.h b/Engine/Source/Runtime/AugmentedReality/Public/ARTrackable.h
--- a/Engine/Source/Runtime/AugmentedReality/Public/ARTrackable.h
+++ b/Engine/Source/Runtime/AugmentedReality/Public/ARTrackable.h
@@ -397,3 +397,21 @@ private:
 	UPROPERTY()
 	UARCandidateObject* DetectedObject;
 };
+
+UCLASS(BlueprintType)
+class AUGMENTEDREALITY_API UARTrackedPose : public UARTrackedGeometry
+{
+	GENERATED_BODY()
+	
+public:
+	virtual void DebugDraw(UWorld* World, const FLinearColor& OutlineColor, float OutlineThickness, float PersistForSeconds = 0.0f) const override;
+	
+	void UpdateTrackedPose(const TSharedRef<FARSupportInterface , ESPMode::ThreadSafe>& InTrackingSystem, uint32 FrameNumber, double Timestamp, const FTransform& InLocalToTrackingTransform, const FTransform& InAlignmentTransform, const FARPose3D& InTrackedPose);
+	
+	UFUNCTION(BlueprintPure, Category = "AR AugmentedReality|Pose Tracking")
+	const FARPose3D& GetTrackedPoseData() const { return TrackedPose; };
+	
+private:
+	/** The detailed info of the tracked pose */
+	FARPose3D TrackedPose;
+};
diff --git a/Engine/Source/Runtime/AugmentedReality/Public/ARTypes.h b/Engine/Source/Runtime/AugmentedReality/Public/ARTypes.h
--- a/Engine/Source/Runtime/AugmentedReality/Public/ARTypes.h
+++ b/Engine/Source/Runtime/AugmentedReality/Public/ARTypes.h
@@ -172,6 +172,23 @@ enum class EARObjectClassification : uint8
 	// Add other types here...
 };
 
+/** Describes the potential spaces in which the join transform can be defined with AR pose tracking */
+UENUM(BlueprintType)
+enum class EARJointTransformSpace : uint8
+{
+	/**
+	 * Joint transform is relative to the origin of the model space
+	 * which is usually attached to a particular joint
+	 * such as the hip
+	 */
+	Model,
+	
+	/**
+	* Joint transform is relative to its parent
+	*/
+	ParentJoint,
+};
+
 /** The current state of the AR subsystem including an optional explanation string. */
 USTRUCT(BlueprintType)
 struct AUGMENTEDREALITY_API FARSessionStatus
@@ -448,3 +465,65 @@ public:
 	}
 };
 
+/** Represents a hierarchy of a human pose skeleton tracked by the AR system */
+USTRUCT(BlueprintType)
+struct AUGMENTEDREALITY_API FARSkeletonDefinition
+{
+	GENERATED_BODY()
+
+public:
+	/** How many joints this skeleton has */
+	UPROPERTY(BlueprintReadOnly, Category="AR AugmentedReality|Pose Tracking")
+	int32 NumJoints = 0;
+	
+	/** The name of each joint in this skeleton */
+	UPROPERTY(BlueprintReadOnly, Category = "AR AugmentedReality|Pose Tracking")
+	TArray<FName> JointNames;
+
+	/** The parent index of each joint in this skeleton */
+	UPROPERTY(BlueprintReadOnly, Category = "AR AugmentedReality|Pose Tracking")
+	TArray<int32> ParentIndices;
+};
+
+/** Represents a human pose tracked in the 2D space */
+USTRUCT(BlueprintType)
+struct AUGMENTEDREALITY_API FARPose2D
+{
+	GENERATED_BODY()
+
+public:
+	/** The definition of this skeleton */
+	UPROPERTY(BlueprintReadOnly, Category="AR AugmentedReality|Pose Tracking")
+	FARSkeletonDefinition SkeletonDefinition;
+	
+	/** The location of each joint in 2D normalized space */
+	UPROPERTY(BlueprintReadOnly, Category = "AR AugmentedReality|Pose Tracking")
+	TArray<FVector2D> JointLocations;
+
+	/** Flags indicating if each joint is tracked */
+	UPROPERTY(BlueprintReadOnly, Category = "AR AugmentedReality|Pose Tracking")
+	TArray<bool> IsJointTracked;
+};
+
+/** Represents a human pose tracked in the 3D space */
+USTRUCT(BlueprintType)
+struct AUGMENTEDREALITY_API FARPose3D
+{
+	GENERATED_BODY()
+
+public:
+	/** The definition of this skeleton */
+	UPROPERTY(BlueprintReadOnly, Category="AR AugmentedReality|Pose Tracking")
+	FARSkeletonDefinition SkeletonDefinition;
+	
+	/** The transform of each join in the model space */
+	UPROPERTY(BlueprintReadOnly, Category = "AR AugmentedReality|Pose Tracking")
+	TArray<FTransform> JointTransforms;
+	
+	/** Flags indicating if each joint is tracked */
+	UPROPERTY(BlueprintReadOnly, Category = "AR AugmentedReality|Pose Tracking")
+	TArray<bool> IsJointTracked;
+	
+	UPROPERTY(BlueprintReadOnly, Category = "AR AugmentedReality|Pose Tracking")
+	EARJointTransformSpace JointTransformSpace = EARJointTransformSpace::Model;
+};
