UE-79885 HL2 CL integrations
-Fixing missing include.

diff --git a/Engine/Source/Runtime/AugmentedReality/Private/AROriginActor.cpp b/Engine/Source/Runtime/AugmentedReality/Private/AROriginActor.cpp
--- a/Engine/Source/Runtime/AugmentedReality/Private/AROriginActor.cpp
+++ b/Engine/Source/Runtime/AugmentedReality/Private/AROriginActor.cpp
@@ -2,6 +2,7 @@
 
 #include "AROriginActor.h"
 #include "EngineUtils.h"
+#include "Engine/Engine.h"
 
 AAROriginActor::AAROriginActor(const FObjectInitializer& ObjectInitializer)
 	: Super(ObjectInitializer)
