Added some "Ignored by ARCore" comments Google requested

diff --git a/Engine/Source/Runtime/AugmentedReality/Public/ARTypes.h b/Engine/Source/Runtime/AugmentedReality/Public/ARTypes.h
--- a/Engine/Source/Runtime/AugmentedReality/Public/ARTypes.h
+++ b/Engine/Source/Runtime/AugmentedReality/Public/ARTypes.h
@@ -201,11 +201,11 @@ private:
 	UPROPERTY(EditAnywhere, Category = "AR Candidate Image")
 	float Width;
 
-	/** The physical height in centimeters of the object that this candidate image represents */
+	/** The physical height in centimeters of the object that this candidate image represents. Ignored in ARCore */
 	UPROPERTY(EditAnywhere, Category = "AR Candidate Image")
 	float Height;
 
-	/** The orientation to treat the candidate image as */
+	/** The orientation to treat the candidate image as. Ignored in ARCore */
 	UPROPERTY(EditAnywhere, Category = "AR Candidate Image")
 	EARCandidateImageOrientation Orientation;
 };
