Removed the ensures from the ARBlueprintLibrary code since the code handles the failure cases properly

diff --git a/Engine/Source/Runtime/AugmentedReality/Private/ARBlueprintLibrary.cpp b/Engine/Source/Runtime/AugmentedReality/Private/ARBlueprintLibrary.cpp
--- a/Engine/Source/Runtime/AugmentedReality/Private/ARBlueprintLibrary.cpp
+++ b/Engine/Source/Runtime/AugmentedReality/Private/ARBlueprintLibrary.cpp
@@ -8,58 +8,8 @@
 #include "ARPin.h"
 #include "ARTrackable.h"
 
-
-
-
-
-//
-//
-//
-//UARPinEventHandlers* UARPinEventHandlers::HandleARPinEvents( UARPin* Pin )
-//{
-//	UARPinEventHandlers* NewEventHandlers = NewObject<UARPinEventHandlers>();
-//	Pin->SetOnTarckingStateChanged(NewEventHandlers->OnARTrackingStateChanged);
-//	Pin->SetOnTransformUpdated(NewEventHandlers->OnARTransformUpdated);
-//	
-//	return NewEventHandlers;
-//}
-
-
-
-
 TSharedPtr<FARSystemBase, ESPMode::ThreadSafe> UARBlueprintLibrary::RegisteredARSystem = nullptr;
 
-//bool UARBlueprintLibrary::ARLineTraceFromScreenPoint(UObject* WorldContextObject, const FVector2D ScreenPosition, TArray<FARTraceResult>& OutHitResults)
-//{
-//	TArray<IARHitTestingSupport*> Providers = IModularFeatures::Get().GetModularFeatureImplementations<IARHitTestingSupport>(IARHitTestingSupport::GetModularFeatureName());
-//	const int32 NumProviders = Providers.Num();
-//	ensureMsgf(NumProviders <= 1, TEXT("Expected at most one ARHitTestingSupport provider, but there are %d registered. Using the first."), NumProviders);
-//	if ( ensureMsgf(NumProviders > 0, TEXT("Expected at least one ARHitTestingSupport provider.")) )
-//	{
-//		if (const bool bHit = Providers[0]->ARLineTraceFromScreenPoint(ScreenPosition, OutHitResults))
-//		{
-//			// Update transform from ARKit (camera) space to UE World Space
-//			
-//			UWorld* MyWorld = WorldContextObject->GetWorld();
-//			APlayerController* MyPC = MyWorld != nullptr ? MyWorld->GetFirstPlayerController() : nullptr;
-//			APawn* MyPawn = MyPC != nullptr ? MyPC->GetPawn() : nullptr;
-//			
-//			if (MyPawn != nullptr)
-//			{
-//				const FTransform PawnTransform = MyPawn->GetActorTransform();
-//				for ( FARTraceResult& HitResult : OutHitResults )
-//				{
-//					HitResult.Transform *= PawnTransform;
-//				}
-//				return true;
-//			}
-//		}
-//	}
-//
-//	return false;
-//}
-
-
 EARTrackingQuality UARBlueprintLibrary::GetTrackingQuality()
 {
 	auto ARSystem = GetARSystem();
@@ -78,16 +28,14 @@ void UARBlueprintLibrary::StartARSession(UARSessionConfig* SessionConfig)
 	static const TCHAR NotARApp_Warning[] = TEXT("Attempting to start an AR session but there is no AR plugin configured. To use AR, enable the proper AR plugin in the Plugin Settings.");
 	
 	auto ARSystem = GetARSystem();
-	if (ensureAlwaysMsgf(ARSystem.IsValid(), NotARApp_Warning))
+	if (ARSystem.IsValid())
 	{
 		ARSystem->StartARSession(SessionConfig);
 	}
 	else
 	{
 #if !(UE_BUILD_SHIPPING || UE_BUILD_TEST)
-		// Ensures don't show up on iOS, but we definitely want a developer to see this
-		// Their AR project just doesn't make sense unless they have enabled the AR setting.
-		GEngine->AddOnScreenDebugMessage(INDEX_NONE, 3600.0f, FColor(255,48,16),NotARApp_Warning);
+		GEngine->AddOnScreenDebugMessage(INDEX_NONE, 3600.0f, FColor(255,48,16), NotARApp_Warning);
 #endif //!(UE_BUILD_SHIPPING || UE_BUILD_TEST)
 	}
 }
@@ -151,7 +99,7 @@ TArray<FARTraceResult> UARBlueprintLibrary::LineTraceTrackedObjects( const FVect
 	TArray<FARTraceResult> Result;
 	
 	auto ARSystem = GetARSystem();
-	if (ensure(ARSystem.IsValid()))
+	if (ARSystem.IsValid())
 	{
 		EARLineTraceChannels ActiveTraceChannels =
 		(bTestFeaturePoints ? EARLineTraceChannels::FeaturePoint : EARLineTraceChannels::None) |
@@ -170,7 +118,7 @@ TArray<UARTrackedGeometry*> UARBlueprintLibrary::GetAllGeometries()
 	TArray<UARTrackedGeometry*> Geometries;
 	
 	auto ARSystem = GetARSystem();
-	if (ensure(ARSystem.IsValid()))
+	if (ARSystem.IsValid())
 	{
 		Geometries = ARSystem->GetAllTrackedGeometries();
 	}
@@ -182,7 +130,7 @@ TArray<UARPin*> UARBlueprintLibrary::GetAllPins()
 	TArray<UARPin*> Pins;
 	
 	auto ARSystem = GetARSystem();
-	if (ensure(ARSystem.IsValid()))
+	if (ARSystem.IsValid())
 	{
 		Pins = ARSystem->GetAllPins();
 	}
@@ -192,7 +140,7 @@ TArray<UARPin*> UARBlueprintLibrary::GetAllPins()
 bool UARBlueprintLibrary::IsSessionTypeSupported(EARSessionType SessionType)
 {
 	auto ARSystem = GetARSystem();
-	if (ensure(ARSystem.IsValid()))
+	if (ARSystem.IsValid())
 	{
 		return ARSystem->IsSessionTypeSupported(SessionType);
 	}
@@ -203,7 +151,7 @@ bool UARBlueprintLibrary::IsSessionTypeSupported(EARSessionType SessionType)
 void UARBlueprintLibrary::DebugDrawTrackedGeometry( UARTrackedGeometry* TrackedGeometry, UObject* WorldContextObject, FLinearColor Color, float OutlineThickness, float PersistForSeconds )
 {
 	UWorld* MyWorld = WorldContextObject->GetWorld();
-	if (ensure(TrackedGeometry != nullptr) && ensure(MyWorld != nullptr))
+	if (TrackedGeometry != nullptr && MyWorld != nullptr)
 	{
 		TrackedGeometry->DebugDraw(MyWorld, Color, OutlineThickness, PersistForSeconds);
 	}
@@ -212,7 +160,7 @@ void UARBlueprintLibrary::DebugDrawTrackedGeometry( UARTrackedGeometry* TrackedG
 void UARBlueprintLibrary::DebugDrawPin( UARPin* ARPin, UObject* WorldContextObject, FLinearColor Color, float Scale, float PersistForSeconds )
 {
 	UWorld* MyWorld = WorldContextObject->GetWorld();
-	if (ensure(ARPin != nullptr) && ensure(MyWorld != nullptr))
+	if (ARPin != nullptr && MyWorld != nullptr)
 	{
 		ARPin->DebugDraw(MyWorld, Color, Scale, PersistForSeconds);
 	}
@@ -222,7 +170,7 @@ void UARBlueprintLibrary::DebugDrawPin( UARPin* ARPin, UObject* WorldContextObje
 UARLightEstimate* UARBlueprintLibrary::GetCurrentLightEstimate()
 {
 	auto ARSystem = GetARSystem();
-	if (ensure(ARSystem.IsValid()))
+	if (ARSystem.IsValid())
 	{
 		return ARSystem->GetCurrentLightEstimate();
 	}
@@ -232,7 +180,7 @@ UARLightEstimate* UARBlueprintLibrary::GetCurrentLightEstimate()
 UARPin* UARBlueprintLibrary::PinComponent( USceneComponent* ComponentToPin, const FTransform& PinToWorldTransform, UARTrackedGeometry* TrackedGeometry, const FName DebugName )
 {
 	auto ARSystem = GetARSystem();
-	if (ensure(ARSystem.IsValid()))
+	if (ARSystem.IsValid())
 	{
 		return ARSystem->PinComponent( ComponentToPin, PinToWorldTransform, TrackedGeometry, DebugName );
 	}
@@ -242,7 +190,7 @@ UARPin* UARBlueprintLibrary::PinComponent( USceneComponent* ComponentToPin, cons
 UARPin* UARBlueprintLibrary::PinComponentToTraceResult( USceneComponent* ComponentToPin, const FARTraceResult& TraceResult, const FName DebugName )
 {
 	auto ARSystem = GetARSystem();
-	if (ensure(ARSystem.IsValid()))
+	if (ARSystem.IsValid())
 	{
 		return ARSystem->PinComponent( ComponentToPin, TraceResult, DebugName );
 	}
@@ -252,7 +200,7 @@ UARPin* UARBlueprintLibrary::PinComponentToTraceResult( USceneComponent* Compone
 void UARBlueprintLibrary::UnpinComponent( USceneComponent* ComponentToUnpin )
 {
 	auto ARSystem = GetARSystem();
-	if (ensure(ARSystem.IsValid()))
+	if (ARSystem.IsValid())
 	{
 		TArray<UARPin*> AllPins = ARSystem->GetAllPins();
 		const int32 AllPinsCount = AllPins.Num();
@@ -270,7 +218,7 @@ void UARBlueprintLibrary::UnpinComponent( USceneComponent* ComponentToUnpin )
 void UARBlueprintLibrary::RemovePin( UARPin* PinToRemove )
 {
 	auto ARSystem = GetARSystem();
-	if (ensure(ARSystem.IsValid()))
+	if (ARSystem.IsValid())
 	{
 		return ARSystem->RemovePin( PinToRemove );
 	}
@@ -321,7 +269,7 @@ UARTextureCameraImage* UARBlueprintLibrary::GetCameraImage()
 	UARTextureCameraImage* Image = nullptr;
 
 	auto ARSystem = GetARSystem();
-	if (ensure(ARSystem.IsValid()))
+	if (ARSystem.IsValid())
 	{
 		Image = ARSystem->GetCameraImage();
 	}
@@ -333,7 +281,7 @@ UARTextureCameraDepth* UARBlueprintLibrary::GetCameraDepth()
 	UARTextureCameraDepth* Depth = nullptr;
 
 	auto ARSystem = GetARSystem();
-	if (ensure(ARSystem.IsValid()))
+	if (ARSystem.IsValid())
 	{
 		Depth = ARSystem->GetCameraDepth();
 	}
@@ -345,7 +293,7 @@ TArray<UARPlaneGeometry*> UARBlueprintLibrary::GetAllTrackedPlanes()
 	TArray<UARPlaneGeometry*> ResultList;
 	
 	auto ARSystem = GetARSystem();
-	if (ensure(ARSystem.IsValid()))
+	if (ARSystem.IsValid())
 	{
 		for (UARTrackedGeometry* Geo : ARSystem->GetAllTrackedGeometries())
 		{
@@ -363,7 +311,7 @@ TArray<UARTrackedPoint*> UARBlueprintLibrary::GetAllTrackedPoints()
 	TArray<UARTrackedPoint*> ResultList;
 	
 	auto ARSystem = GetARSystem();
-	if (ensure(ARSystem.IsValid()))
+	if (ARSystem.IsValid())
 	{
 		for (UARTrackedGeometry* Geo : ARSystem->GetAllTrackedGeometries())
 		{
@@ -381,7 +329,7 @@ TArray<UARTrackedImage*> UARBlueprintLibrary::GetAllTrackedImages()
 	TArray<UARTrackedImage*> ResultList;
 	
 	auto ARSystem = GetARSystem();
-	if (ensure(ARSystem.IsValid()))
+	if (ARSystem.IsValid())
 	{
 		for (UARTrackedGeometry* Geo : ARSystem->GetAllTrackedGeometries())
 		{
@@ -399,7 +347,7 @@ TArray<UAREnvironmentCaptureProbe*> UARBlueprintLibrary::GetAllTrackedEnvironmen
 	TArray<UAREnvironmentCaptureProbe*> ResultList;
 	
 	auto ARSystem = GetARSystem();
-	if (ensure(ARSystem.IsValid()))
+	if (ARSystem.IsValid())
 	{
 		for (UARTrackedGeometry* Geo : ARSystem->GetAllTrackedGeometries())
 		{
@@ -415,7 +363,7 @@ TArray<UAREnvironmentCaptureProbe*> UARBlueprintLibrary::GetAllTrackedEnvironmen
 bool UARBlueprintLibrary::AddManualEnvironmentCaptureProbe(FVector Location, FVector Extent)
 {
 	auto ARSystem = GetARSystem();
-	if (ensure(ARSystem.IsValid()))
+	if (ARSystem.IsValid())
 	{
 		return ARSystem->AddManualEnvironmentCaptureProbe(Location, Extent);
 	}
@@ -425,7 +373,7 @@ bool UARBlueprintLibrary::AddManualEnvironmentCaptureProbe(FVector Location, FVe
 EARWorldMappingState UARBlueprintLibrary::GetWorldMappingStatus()
 {
 	auto ARSystem = GetARSystem();
-	if (ensure(ARSystem.IsValid()))
+	if (ARSystem.IsValid())
 	{
 		return ARSystem->GetWorldMappingStatus();
 	}
@@ -435,7 +383,7 @@ EARWorldMappingState UARBlueprintLibrary::GetWorldMappingStatus()
 TSharedPtr<FARSaveWorldAsyncTask, ESPMode::ThreadSafe> UARBlueprintLibrary::SaveWorld()
 {
 	auto ARSystem = GetARSystem();
-	if (ensure(ARSystem.IsValid()))
+	if (ARSystem.IsValid())
 	{
 		return ARSystem->SaveWorld();
 		
@@ -446,7 +394,7 @@ TSharedPtr<FARSaveWorldAsyncTask, ESPMode::ThreadSafe> UARBlueprintLibrary::Save
 TSharedPtr<FARGetCandidateObjectAsyncTask, ESPMode::ThreadSafe> UARBlueprintLibrary::GetCandidateObject(FVector Location, FVector Extent)
 {
 	auto ARSystem = GetARSystem();
-	if (ensure(ARSystem.IsValid()))
+	if (ARSystem.IsValid())
 	{
 		return ARSystem->GetCandidateObject(Location, Extent);
 		
@@ -457,7 +405,7 @@ TSharedPtr<FARGetCandidateObjectAsyncTask, ESPMode::ThreadSafe> UARBlueprintLibr
 TArray<FARVideoFormat> UARBlueprintLibrary::GetSupportedVideoFormats(EARSessionType SessionType)
 {
 	auto ARSystem = GetARSystem();
-	if (ensure(ARSystem.IsValid()))
+	if (ARSystem.IsValid())
 	{
 		return ARSystem->GetSupportedVideoFormats(SessionType);
 	}
