Fix for include file casing

diff --git a/Engine/Source/Runtime/AugmentedReality/Private/ARSupportInterface.cpp b/Engine/Source/Runtime/AugmentedReality/Private/ARSupportInterface.cpp
--- a/Engine/Source/Runtime/AugmentedReality/Private/ARSupportInterface.cpp
+++ b/Engine/Source/Runtime/AugmentedReality/Private/ARSupportInterface.cpp
@@ -1,7 +1,7 @@
 // Copyright 1998-2018 Epic Games, Inc. All Rights Reserved.
 
 #include "ARSupportInterface.h"
-#include "ArTraceResult.h"
+#include "ARTraceResult.h"
 #include "ARSystem.h"
 #include "Features/IModularFeatures.h"
 #include "ARBlueprintLibrary.h"
