Added delegates so that AR apps can use an event based approach rather than polling for changes

diff --git a/Engine/Plugins/Lumin/MagicLeap/Source/MagicLeap/Private/MagicLeapHMD.h b/Engine/Plugins/Lumin/MagicLeap/Source/MagicLeap/Private/MagicLeapHMD.h
--- a/Engine/Plugins/Lumin/MagicLeap/Source/MagicLeap/Private/MagicLeapHMD.h
+++ b/Engine/Plugins/Lumin/MagicLeap/Source/MagicLeap/Private/MagicLeapHMD.h
@@ -25,7 +25,6 @@ ML_INCLUDES_START
 ML_INCLUDES_END
 #endif //WITH_MLSDK
 
-class IARSystemSupport;
 class FXRTrackingSystemBase;
 
 /**
diff --git a/Engine/Plugins/Runtime/AR/Apple/AppleARKit/Source/AppleARKit/Private/AppleARKitSystem.cpp b/Engine/Plugins/Runtime/AR/Apple/AppleARKit/Source/AppleARKit/Private/AppleARKitSystem.cpp
--- a/Engine/Plugins/Runtime/AR/Apple/AppleARKit/Source/AppleARKit/Private/AppleARKitSystem.cpp
+++ b/Engine/Plugins/Runtime/AR/Apple/AppleARKit/Source/AppleARKit/Private/AppleARKitSystem.cpp
@@ -206,6 +206,7 @@ FAppleARKitSystem::FAppleARKitSystem()
 , CameraDepth(nullptr)
 , LastTrackedGeometry_DebugId(0)
 , FaceARSupport(nullptr)
+, TimecodeProvider(nullptr)
 {
 	// See Initialize(), as we need access to SharedThis()
 }
@@ -1702,6 +1703,9 @@ PRAGMA_ENABLE_DEPRECATION_WARNINGS
 	UARTrackedGeometry* NewTrackedGeometry = TrackedGeometries.Add( AnchorData->AnchorGUID, NewGeometry );
 	
 	NewTrackedGeometry->SetDebugName( FName(*NewAnchorDebugName) );
+
+	// Trigger the delegate so anyone listening can take action
+	TriggerOnTrackableAddedDelegates(NewTrackedGeometry);
 }
 
 void FAppleARKitSystem::SessionDidUpdateAnchors_Internal( TSharedRef<FAppleARKitAnchorData> AnchorData )
@@ -1816,6 +1820,8 @@ PRAGMA_ENABLE_DEPRECATION_WARNINGS
 				break;
 			}
 		}
+		// Trigger the delegate so anyone listening can take action
+		TriggerOnTrackableUpdatedDelegates(FoundGeometry);
 	}
 }
 
@@ -1839,6 +1845,8 @@ void FAppleARKitSystem::SessionDidRemoveAnchors_Internal( FGuid AnchorGuid )
 		{
 			PinBeingOrphaned->OnTrackingStateChanged(EARTrackingState::StoppedTracking);
 		}
+		// Trigger the delegate so anyone listening can take action
+		TriggerOnTrackableRemovedDelegates(TrackedGeometryBeingRemoved);
 	}
 	
 	TrackedGeometries.Remove(AnchorGuid);
diff --git a/Engine/Source/Runtime/AugmentedReality/Private/ARSupportInterface.cpp b/Engine/Source/Runtime/AugmentedReality/Private/ARSupportInterface.cpp
--- a/Engine/Source/Runtime/AugmentedReality/Private/ARSupportInterface.cpp
+++ b/Engine/Source/Runtime/AugmentedReality/Private/ARSupportInterface.cpp
@@ -2,14 +2,13 @@
 
 #include "ARSupportInterface.h"
 #include "ARTraceResult.h"
-#include "ARSystem.h"
 #include "Features/IModularFeatures.h"
 #include "ARBlueprintLibrary.h"
 #include "ARBlueprintProxy.h"
 #include "Templates/SharedPointer.h"
 #include "Engine/Texture2D.h"
 
-FARSupportInterface ::FARSupportInterface (IARSystemSupport* InARImplementation, IXRTrackingSystem* InXRTrackingSystem)
+FARSupportInterface::FARSupportInterface (IARSystemSupport* InARImplementation, IXRTrackingSystem* InXRTrackingSystem)
 	: ARImplemention(InARImplementation)
 	, XRTrackingSystem(InXRTrackingSystem)
 	, AlignmentTransform(FTransform::Identity)
@@ -323,3 +322,33 @@ void FARSupportInterface ::AddReferencedObjects(FReferenceCollector& Collector)
 		Collector.AddReferencedObject(ARSettings);
 	}
 }
+
+#define DEFINE_AR_SI_DELEGATE_FUNCS(DelegateName) \
+FDelegateHandle FARSupportInterface::Add##DelegateName##Delegate_Handle(const F##DelegateName##Delegate& Delegate) \
+{ \
+	if (ARImplemention) \
+	{ \
+		return ARImplemention->Add##DelegateName##Delegate_Handle(Delegate); \
+	} \
+	return Delegate.GetHandle(); \
+} \
+void FARSupportInterface::Clear##DelegateName##Delegate_Handle(FDelegateHandle& Handle) \
+{ \
+	if (ARImplemention) \
+	{ \
+		ARImplemention->Clear##DelegateName##Delegate_Handle(Handle); \
+		return; \
+	} \
+	Handle.Reset(); \
+} \
+void FARSupportInterface::Clear##DelegateName##Delegates(void* Object) \
+{ \
+	if (ARImplemention) \
+	{ \
+		ARImplemention->Clear##DelegateName##Delegates(Object); \
+	} \
+}
+
+DEFINE_AR_SI_DELEGATE_FUNCS(OnTrackableAdded)
+DEFINE_AR_SI_DELEGATE_FUNCS(OnTrackableUpdated)
+DEFINE_AR_SI_DELEGATE_FUNCS(OnTrackableRemoved)
diff --git a/Engine/Source/Runtime/AugmentedReality/Private/ARTrackableNotifyComponent.cpp b/Engine/Source/Runtime/AugmentedReality/Private/ARTrackableNotifyComponent.cpp
new file mode 100644
--- /dev/null
+++ b/Engine/Source/Runtime/AugmentedReality/Private/ARTrackableNotifyComponent.cpp
@@ -0,0 +1,140 @@
+// Copyright 1998-2019 Epic Games, Inc. All Rights Reserved.
+
+#include "ARTrackableNotifyComponent.h"
+#include "ARBlueprintLibrary.h"
+
+void UARTrackableNotifyComponent::OnRegister()
+{
+	Super::OnRegister();
+
+	// Add the delegates we want to listen for
+	UARBlueprintLibrary::AddOnTrackableAddedDelegate_Handle(FOnTrackableAddedDelegate::CreateUObject(this, &UARTrackableNotifyComponent::OnTrackableAdded));
+	UARBlueprintLibrary::AddOnTrackableUpdatedDelegate_Handle(FOnTrackableAddedDelegate::CreateUObject(this, &UARTrackableNotifyComponent::OnTrackableUpdated));
+	UARBlueprintLibrary::AddOnTrackableRemovedDelegate_Handle(FOnTrackableAddedDelegate::CreateUObject(this, &UARTrackableNotifyComponent::OnTrackableRemoved));
+}
+
+void UARTrackableNotifyComponent::OnUnregister()
+{
+	// Remove the delegates we are listening for
+	UARBlueprintLibrary::ClearOnTrackableAddedDelegates(this);
+	UARBlueprintLibrary::ClearOnTrackableUpdatedDelegates(this);
+	UARBlueprintLibrary::ClearOnTrackableRemovedDelegates(this);
+
+	Super::OnUnregister();
+}
+
+template<typename OBJ_TYPE, typename DELEGATE_TYPE>
+bool UARTrackableNotifyComponent::ConditionalDispatchEvent(UARTrackedGeometry* Tracked, DELEGATE_TYPE& Delegate)
+{
+	if (OBJ_TYPE* CastObject = Cast<OBJ_TYPE>(Tracked))
+	{
+		if (Delegate.IsBound())
+		{
+			Delegate.Broadcast(CastObject);
+		}
+		return true;
+	}
+	return false;
+}
+
+void UARTrackableNotifyComponent::OnTrackableAdded(UARTrackedGeometry* Added)
+{
+	if (ConditionalDispatchEvent<UARTrackedObject, FTrackableObjectDelegate>(Added, OnAddTrackedObject))
+	{
+		// Do nothing, since we found the right type
+	}
+	else if (ConditionalDispatchEvent<UAREnvironmentCaptureProbe, FTrackableEnvProbeDelegate>(Added, OnAddTrackedEnvProbe))
+	{
+		// Do nothing, since we found the right type
+	}
+	else if (ConditionalDispatchEvent<UARFaceGeometry, FTrackableFaceDelegate>(Added, OnAddTrackedFace))
+	{
+		// Do nothing, since we found the right type
+	}
+	else if (ConditionalDispatchEvent<UARTrackedImage, FTrackableImageDelegate>(Added, OnAddTrackedImage))
+	{
+		// Do nothing, since we found the right type
+	}
+	else if (ConditionalDispatchEvent<UARTrackedPoint, FTrackablePointDelegate>(Added, OnAddTrackedPoint))
+	{
+		// Do nothing, since we found the right type
+	}
+	else if (ConditionalDispatchEvent<UARPlaneGeometry, FTrackablePlaneDelegate>(Added, OnAddTrackedPlane))
+	{
+		// Do nothing, since we found the right type
+	}
+
+	// Always trigger the base version since it matches all trackables
+	if (OnAddTrackedGeometry.IsBound())
+	{
+		OnAddTrackedGeometry.Broadcast(Added);
+	}
+}
+
+void UARTrackableNotifyComponent::OnTrackableUpdated(UARTrackedGeometry* Updated)
+{
+	if (ConditionalDispatchEvent<UARTrackedObject, FTrackableObjectDelegate>(Updated, OnUpdateTrackedObject))
+	{
+		// Do nothing, since we found the right type
+	}
+	else if (ConditionalDispatchEvent<UAREnvironmentCaptureProbe, FTrackableEnvProbeDelegate>(Updated, OnUpdateTrackedEnvProbe))
+	{
+		// Do nothing, since we found the right type
+	}
+	else if (ConditionalDispatchEvent<UARFaceGeometry, FTrackableFaceDelegate>(Updated, OnUpdateTrackedFace))
+	{
+		// Do nothing, since we found the right type
+	}
+	else if (ConditionalDispatchEvent<UARTrackedImage, FTrackableImageDelegate>(Updated, OnUpdateTrackedImage))
+	{
+		// Do nothing, since we found the right type
+	}
+	else if (ConditionalDispatchEvent<UARTrackedPoint, FTrackablePointDelegate>(Updated, OnUpdateTrackedPoint))
+	{
+		// Do nothing, since we found the right type
+	}
+	else if (ConditionalDispatchEvent<UARPlaneGeometry, FTrackablePlaneDelegate>(Updated, OnUpdateTrackedPlane))
+	{
+		// Do nothing, since we found the right type
+	}
+
+	// Always trigger the base version since it matches all trackables
+	if (OnUpdateTrackedGeometry.IsBound())
+	{
+		OnUpdateTrackedGeometry.Broadcast(Updated);
+	}
+}
+
+void UARTrackableNotifyComponent::OnTrackableRemoved(UARTrackedGeometry* Removed)
+{
+	if (ConditionalDispatchEvent<UARTrackedObject, FTrackableObjectDelegate>(Removed, OnRemoveTrackedObject))
+	{
+		// Do nothing, since we found the right type
+	}
+	else if (ConditionalDispatchEvent<UAREnvironmentCaptureProbe, FTrackableEnvProbeDelegate>(Removed, OnRemoveTrackedEnvProbe))
+	{
+		// Do nothing, since we found the right type
+	}
+	else if (ConditionalDispatchEvent<UARFaceGeometry, FTrackableFaceDelegate>(Removed, OnRemoveTrackedFace))
+	{
+		// Do nothing, since we found the right type
+	}
+	else if (ConditionalDispatchEvent<UARTrackedImage, FTrackableImageDelegate>(Removed, OnRemoveTrackedImage))
+	{
+		// Do nothing, since we found the right type
+	}
+	else if (ConditionalDispatchEvent<UARTrackedPoint, FTrackablePointDelegate>(Removed, OnRemoveTrackedPoint))
+	{
+		// Do nothing, since we found the right type
+	}
+	else if (ConditionalDispatchEvent<UARPlaneGeometry, FTrackablePlaneDelegate>(Removed, OnRemoveTrackedPlane))
+	{
+		// Do nothing, since we found the right type
+	}
+
+	// Always trigger the base version since it matches all trackables
+	if (OnRemoveTrackedGeometry.IsBound())
+	{
+		OnRemoveTrackedGeometry.Broadcast(Removed);
+	}
+}
diff --git a/Engine/Source/Runtime/AugmentedReality/Public/ARBlueprintLibrary.h b/Engine/Source/Runtime/AugmentedReality/Public/ARBlueprintLibrary.h
--- a/Engine/Source/Runtime/AugmentedReality/Public/ARBlueprintLibrary.h
+++ b/Engine/Source/Runtime/AugmentedReality/Public/ARBlueprintLibrary.h
@@ -6,12 +6,42 @@
 #include "Kismet/BlueprintFunctionLibrary.h"
 #include "Kismet/BlueprintAsyncActionBase.h"
 #include "ARTypes.h"
+#include "ARSystem.h"
 #include "ARTraceResult.h"
 #include "ARSessionConfig.h"
 #include "ARTrackable.h"
 #include "ARBlueprintLibrary.generated.h"
 
 
+#define DEFINE_AR_BPLIB_DELEGATE_FUNCS(DelegateName) \
+public: \
+	static FDelegateHandle Add##DelegateName##Delegate_Handle(const F##DelegateName##Delegate& Delegate) \
+	{ \
+		auto ARSystem = GetARSystem(); \
+		if (ARSystem.IsValid()) \
+		{ \
+			return ARSystem.Pin()->Add##DelegateName##Delegate_Handle(Delegate); \
+		} \
+		return Delegate.GetHandle(); \
+	} \
+	static void Clear##DelegateName##Delegate_Handle(FDelegateHandle& Handle) \
+	{ \
+		auto ARSystem = GetARSystem(); \
+		if (ARSystem.IsValid()) \
+		{ \
+			ARSystem.Pin()->Clear##DelegateName##Delegate_Handle(Handle); \
+			return; \
+		} \
+		Handle.Reset(); \
+	} \
+	static void Clear##DelegateName##Delegates(void* Object) \
+	{ \
+		auto ARSystem = GetARSystem(); \
+		if (ARSystem.IsValid()) \
+		{ \
+			ARSystem.Pin()->Clear##DelegateName##Delegates(Object); \
+		} \
+	}
 
 UCLASS(meta=(ScriptName="ARLibrary"))
 class AUGMENTEDREALITY_API UARBlueprintLibrary : public UBlueprintFunctionLibrary
@@ -230,6 +260,12 @@ public:
 	UFUNCTION(BlueprintCallable, Category = "AR AugmentedReality|Session", meta = (Keywords = "ar augmentedreality augmented reality candidate image"))
 	static UARCandidateImage* AddRuntimeCandidateImage(UARSessionConfig* SessionConfig, UTexture2D* CandidateTexture, FString FriendlyName, float PhysicalWidth);
 
+	// Static helpers to create the methods needed to add/remove delegates from the AR system
+	DEFINE_AR_BPLIB_DELEGATE_FUNCS(OnTrackableAdded)
+	DEFINE_AR_BPLIB_DELEGATE_FUNCS(OnTrackableUpdated)
+	DEFINE_AR_BPLIB_DELEGATE_FUNCS(OnTrackableRemoved)
+	// End helpers
+
 public:
 	static void RegisterAsARSystem(const TSharedRef<FARSupportInterface , ESPMode::ThreadSafe>& NewArSystem);
 	
diff --git a/Engine/Source/Runtime/AugmentedReality/Public/ARSupportInterface.h b/Engine/Source/Runtime/AugmentedReality/Public/ARSupportInterface.h
--- a/Engine/Source/Runtime/AugmentedReality/Public/ARSupportInterface.h
+++ b/Engine/Source/Runtime/AugmentedReality/Public/ARSupportInterface.h
@@ -12,6 +12,12 @@ class IARSystemSupport;
 DECLARE_MULTICAST_DELEGATE(FARSystemOnSessionStarted);
 DECLARE_MULTICAST_DELEGATE_OneParam(FARSystemOnAlignmentTransformUpdated, const FTransform&);
 
+#define DECLARE_AR_SI_DELEGATE_FUNCS(DelegateName) \
+public: \
+	FDelegateHandle Add##DelegateName##Delegate_Handle(const F##DelegateName##Delegate& Delegate); \
+	void Clear##DelegateName##Delegate_Handle(FDelegateHandle& Handle); \
+	void Clear##DelegateName##Delegates(void* Object);
+
 /**
  * Composition Components for tracking system features
 */
@@ -107,6 +113,12 @@ public:
 	virtual void AddReferencedObjects(FReferenceCollector& Collector) override;
 	//~ FGCObject
 
+	// Pass through helpers to create the methods needed to add/remove delegates from the AR system
+	DECLARE_AR_SI_DELEGATE_FUNCS(OnTrackableAdded)
+	DECLARE_AR_SI_DELEGATE_FUNCS(OnTrackableUpdated)
+	DECLARE_AR_SI_DELEGATE_FUNCS(OnTrackableRemoved)
+	// End helpers
+
 	FARSystemOnSessionStarted OnARSessionStarted;
 	FARSystemOnAlignmentTransformUpdated OnAlignmentTransformUpdated;
 
diff --git a/Engine/Source/Runtime/AugmentedReality/Public/ARSystem.h b/Engine/Source/Runtime/AugmentedReality/Public/ARSystem.h
--- a/Engine/Source/Runtime/AugmentedReality/Public/ARSystem.h
+++ b/Engine/Source/Runtime/AugmentedReality/Public/ARSystem.h
@@ -2,6 +2,7 @@
 
 #pragma once
 
+#include "CoreMinimal.h"
 #include "XRTrackingSystemBase.h"
 #include "ARTypes.h"
 #include "ARSessionConfig.h"
@@ -14,6 +15,34 @@ class UARTextureCameraImage;
 class UARTextureCameraDepth;
 struct FARTraceResult;
 
+#define DEFINE_AR_DELEGATE_BASE(DelegateName) \
+private: \
+	F##DelegateName DelegateName##Delegates; \
+public: \
+	virtual FDelegateHandle Add##DelegateName##Delegate_Handle(const F##DelegateName##Delegate& Delegate) \
+	{ \
+		DelegateName##Delegates.Add(Delegate); \
+		return Delegate.GetHandle(); \
+	} \
+	virtual void Clear##DelegateName##Delegate_Handle(FDelegateHandle& Handle) \
+	{ \
+		DelegateName##Delegates.Remove(Handle); \
+		Handle.Reset(); \
+	} \
+	virtual void Clear##DelegateName##Delegates(void* Object) \
+	{ \
+		DelegateName##Delegates.RemoveAll(Object); \
+	}
+
+#define DEFINE_AR_DELEGATE_ONE_PARAM(DelegateName, Param1Type) \
+	DEFINE_AR_DELEGATE_BASE(DelegateName) \
+protected: \
+	virtual void Trigger##DelegateName##Delegates(Param1Type Param1) \
+	{ \
+		DelegateName##Delegates.Broadcast(Param1); \
+	} \
+public:
+
 
 /**
  * Implement IARSystemSupport for any platform that wants to be an Unreal Augmented Reality System. e.g. AppleARKit, GoogleARCore.
@@ -151,10 +180,29 @@ public:
 
 	virtual void* GetARSessionRawPointer() = 0;
 	virtual void* GetGameThreadARFrameRawPointer() = 0;
-	
 
-public:
 	virtual ~IARSystemSupport(){}
+
+	/**
+	 * Delegate called when an ar item is added to the scene
+	 *
+	 * @param Added the item that was added
+	 */
+	DEFINE_AR_DELEGATE_ONE_PARAM(OnTrackableAdded, UARTrackedGeometry* /* Added */);
+
+	/**
+	 * Delegate called when an ar item is updated
+	 *
+	 * @param Updated the item that was updated
+	 */
+	DEFINE_AR_DELEGATE_ONE_PARAM(OnTrackableUpdated, UARTrackedGeometry* /* Updated */);
+
+	/**
+	 * Delegate called when an ar item is removed from the scene
+	 *
+	 * @param Removed the item that was removed
+	 */
+	DEFINE_AR_DELEGATE_ONE_PARAM(OnTrackableRemoved, UARTrackedGeometry* /* Removed */);
 };
 
 
diff --git a/Engine/Source/Runtime/AugmentedReality/Public/ARTrackableNotifyComponent.h b/Engine/Source/Runtime/AugmentedReality/Public/ARTrackableNotifyComponent.h
new file mode 100644
--- /dev/null
+++ b/Engine/Source/Runtime/AugmentedReality/Public/ARTrackableNotifyComponent.h
@@ -0,0 +1,147 @@
+// Copyright 1998-2019 Epic Games, Inc. All Rights Reserved.
+
+#pragma once
+
+#include "CoreMinimal.h"
+#include "ARTrackable.h"
+#include "Components/ActorComponent.h"
+#include "ARTrackableNotifyComponent.generated.h"
+
+DECLARE_DYNAMIC_MULTICAST_DELEGATE_OneParam(FTrackableDelegate, UARTrackedGeometry*, TrackedGeometry);
+DECLARE_DYNAMIC_MULTICAST_DELEGATE_OneParam(FTrackablePlaneDelegate, UARPlaneGeometry*, TrackedPlane);
+DECLARE_DYNAMIC_MULTICAST_DELEGATE_OneParam(FTrackablePointDelegate, UARTrackedPoint*, TrackedPoint);
+DECLARE_DYNAMIC_MULTICAST_DELEGATE_OneParam(FTrackableImageDelegate, UARTrackedImage*, TrackedImage);
+DECLARE_DYNAMIC_MULTICAST_DELEGATE_OneParam(FTrackableFaceDelegate, UARFaceGeometry*, TrackedFace);
+DECLARE_DYNAMIC_MULTICAST_DELEGATE_OneParam(FTrackableEnvProbeDelegate, UAREnvironmentCaptureProbe*, TrackedEnvProbe);
+DECLARE_DYNAMIC_MULTICAST_DELEGATE_OneParam(FTrackableObjectDelegate, UARTrackedObject*, TrackedObject);
+
+/** Component used to listen to ar trackable object changes */
+UCLASS(meta=(BlueprintSpawnableComponent))
+class AUGMENTEDREALITY_API UARTrackableNotifyComponent :
+	public UActorComponent
+{
+	GENERATED_BODY()
+
+public:
+// Base class catch all
+	/** Event that happens when any new trackable ar item is added */
+	UPROPERTY(BlueprintAssignable, Category="Event")
+	FTrackableDelegate OnAddTrackedGeometry;
+
+	/** Event that happens when any trackable ar item is updated */
+	UPROPERTY(BlueprintAssignable, Category="Event")
+	FTrackableDelegate OnUpdateTrackedGeometry;
+
+	/** Event that happens when any trackable ar item is removed from the scene */
+	UPROPERTY(BlueprintAssignable, Category="Event")
+	FTrackableDelegate OnRemoveTrackedGeometry;
+
+// Tracked plane events
+	/** Event that happens when any new ar plane item is added */
+	UPROPERTY(BlueprintAssignable, Category="Event")
+	FTrackablePlaneDelegate OnAddTrackedPlane;
+
+	/** Event that happens when any ar plane item is updated */
+	UPROPERTY(BlueprintAssignable, Category="Event")
+	FTrackablePlaneDelegate OnUpdateTrackedPlane;
+
+	/** Event that happens when any ar plane item is removed from the scene */
+	UPROPERTY(BlueprintAssignable, Category="Event")
+	FTrackablePlaneDelegate OnRemoveTrackedPlane;
+
+// Tracked Point events
+	/** Event that happens when any new ar Point item is added */
+	UPROPERTY(BlueprintAssignable, Category="Event")
+	FTrackablePointDelegate OnAddTrackedPoint;
+
+	/** Event that happens when any ar Point item is updated */
+	UPROPERTY(BlueprintAssignable, Category="Event")
+	FTrackablePointDelegate OnUpdateTrackedPoint;
+
+	/** Event that happens when any ar Point item is removed from the scene */
+	UPROPERTY(BlueprintAssignable, Category="Event")
+	FTrackablePointDelegate OnRemoveTrackedPoint;
+
+// Tracked Image events
+	/** Event that happens when any new ar Image item is added */
+	UPROPERTY(BlueprintAssignable, Category="Event")
+	FTrackableImageDelegate OnAddTrackedImage;
+
+	/** Event that happens when any ar Image item is updated */
+	UPROPERTY(BlueprintAssignable, Category="Event")
+	FTrackableImageDelegate OnUpdateTrackedImage;
+
+	/** Event that happens when any ar Image item is removed from the scene */
+	UPROPERTY(BlueprintAssignable, Category="Event")
+	FTrackableImageDelegate OnRemoveTrackedImage;
+
+// Tracked Face events
+	/** Event that happens when any new ar Face item is added */
+	UPROPERTY(BlueprintAssignable, Category="Event")
+	FTrackableFaceDelegate OnAddTrackedFace;
+
+	/** Event that happens when any ar Face item is updated */
+	UPROPERTY(BlueprintAssignable, Category="Event")
+	FTrackableFaceDelegate OnUpdateTrackedFace;
+
+	/** Event that happens when any ar Face item is removed from the scene */
+	UPROPERTY(BlueprintAssignable, Category="Event")
+	FTrackableFaceDelegate OnRemoveTrackedFace;
+
+// Tracked environment capture probe events
+	/** Event that happens when any new ar environment capture probe item is added */
+	UPROPERTY(BlueprintAssignable, Category="Event")
+	FTrackableEnvProbeDelegate OnAddTrackedEnvProbe;
+
+	/** Event that happens when any ar environment capture probe item is updated */
+	UPROPERTY(BlueprintAssignable, Category="Event")
+	FTrackableEnvProbeDelegate OnUpdateTrackedEnvProbe;
+
+	/** Event that happens when any ar environment capture probe item is removed from the scene */
+	UPROPERTY(BlueprintAssignable, Category="Event")
+	FTrackableEnvProbeDelegate OnRemoveTrackedEnvProbe;
+
+// Tracked object events
+	/** Event that happens when any new ar detected object item is added */
+	UPROPERTY(BlueprintAssignable, Category="Event")
+	FTrackableObjectDelegate OnAddTrackedObject;
+
+	/** Event that happens when any ar detected object item is updated */
+	UPROPERTY(BlueprintAssignable, Category="Event")
+	FTrackableObjectDelegate OnUpdateTrackedObject;
+
+	/** Event that happens when any ar detected object item is removed from the scene */
+	UPROPERTY(BlueprintAssignable, Category="Event")
+	FTrackableObjectDelegate OnRemoveTrackedObject;
+
+private:
+	virtual void OnRegister() override;
+	virtual void OnUnregister() override;
+
+	/**
+	 * Triggers the delegate type if the object is the desired class
+	 *
+	 * #return true if the class matched, false otherwise
+	 */
+	template<typename OBJ_TYPE, typename DELEGATE_TYPE>
+	bool ConditionalDispatchEvent(UARTrackedGeometry* Tracked, DELEGATE_TYPE& Delegate);
+
+	/**
+	 * Routes the notification to the proper delegate set
+	 *
+	 * @param Added the newly added ar trackable
+	 */
+	void OnTrackableAdded(UARTrackedGeometry* Added);
+	/**
+	 * Routes the notification to the proper delegate set
+	 *
+	 * @param Updated the updated ar trackable
+	 */
+	void OnTrackableUpdated(UARTrackedGeometry* Updated);
+	/**
+	 * Routes the notification to the proper delegate set
+	 *
+	 * @param Removed the removed ar trackable
+	 */
+	void OnTrackableRemoved(UARTrackedGeometry* Removed);
+};
diff --git a/Engine/Source/Runtime/AugmentedReality/Public/ARTypes.h b/Engine/Source/Runtime/AugmentedReality/Public/ARTypes.h
--- a/Engine/Source/Runtime/AugmentedReality/Public/ARTypes.h
+++ b/Engine/Source/Runtime/AugmentedReality/Public/ARTypes.h
@@ -14,6 +14,16 @@ class UARLightEstimate;
 struct FARTraceResult;
 class UTexture2D;
 
+DECLARE_MULTICAST_DELEGATE_OneParam(FOnTrackableAdded, UARTrackedGeometry*);
+typedef FOnTrackableAdded::FDelegate FOnTrackableAddedDelegate;
+
+DECLARE_MULTICAST_DELEGATE_OneParam(FOnTrackableUpdated, UARTrackedGeometry*);
+typedef FOnTrackableUpdated::FDelegate FOnTrackableUpdatedDelegate;
+
+DECLARE_MULTICAST_DELEGATE_OneParam(FOnTrackableRemoved, UARTrackedGeometry*);
+typedef FOnTrackableRemoved::FDelegate FOnTrackableRemovedDelegate;
+
+
 UENUM(BlueprintType, Category="AR AugmentedReality", meta=(Experimental))
 enum class EARTrackingState : uint8
 {
