Added support for the AR platform's unique id for tracked items

diff --git a/Engine/Plugins/Runtime/AR/Apple/AppleARKit/Source/AppleARKit/Private/AppleARKitSystem.cpp b/Engine/Plugins/Runtime/AR/Apple/AppleARKit/Source/AppleARKit/Private/AppleARKitSystem.cpp
--- a/Engine/Plugins/Runtime/AR/Apple/AppleARKit/Source/AppleARKit/Private/AppleARKitSystem.cpp
+++ b/Engine/Plugins/Runtime/AR/Apple/AppleARKit/Source/AppleARKit/Private/AppleARKitSystem.cpp
@@ -1701,7 +1701,8 @@ PRAGMA_ENABLE_DEPRECATION_WARNINGS
 	check(NewGeometry != nullptr);
 
 	UARTrackedGeometry* NewTrackedGeometry = TrackedGeometries.Add( AnchorData->AnchorGUID, NewGeometry );
-	
+
+	NewTrackedGeometry->UniqueId = AnchorData->AnchorGUID;
 	NewTrackedGeometry->SetDebugName( FName(*NewAnchorDebugName) );
 
 	// Trigger the delegate so anyone listening can take action
diff --git a/Engine/Source/Runtime/AugmentedReality/Public/ARTrackable.h b/Engine/Source/Runtime/AugmentedReality/Public/ARTrackable.h
--- a/Engine/Source/Runtime/AugmentedReality/Public/ARTrackable.h
+++ b/Engine/Source/Runtime/AugmentedReality/Public/ARTrackable.h
@@ -53,7 +53,10 @@ public:
 	
 	UFUNCTION(BlueprintPure, Category="AR AugmentedReality|Tracked Geometry")
 	float GetLastUpdateTimestamp() const;
-	
+
+	UPROPERTY(BlueprintReadOnly, Category="AR AugmentedReality|Tracked Geometry")
+	FGuid UniqueId;
+
 protected:
 	TSharedPtr<FARSupportInterface , ESPMode::ThreadSafe> GetARSystem() const;
 	
@@ -73,12 +76,14 @@ private:
 	TWeakPtr<FARSupportInterface , ESPMode::ThreadSafe> ARSystem;
 	
 	/** The frame number this tracked geometry was last updated on */
-	uint32 LastUpdateFrameNumber;
+	UPROPERTY()
+	int32 LastUpdateFrameNumber;
 	
 	/** The time reported by the AR system that this object was last updated */
 	double LastUpdateTimestamp;
 	
 	/** A unique name that can be used to identify the anchor for debug purposes */
+	UPROPERTY()
 	FName DebugName;
 };
 
