MERGE: Added a AROriginActor for AR systems to spawn when they need to place UActorComponent derived components into the world

diff --git a/Engine/Source/Runtime/AugmentedReality/Private/AROriginActor.cpp b/Engine/Source/Runtime/AugmentedReality/Private/AROriginActor.cpp
new file mode 100644
--- /dev/null
+++ b/Engine/Source/Runtime/AugmentedReality/Private/AROriginActor.cpp
@@ -0,0 +1,8 @@
+// Copyright 1998-2019 Epic Games, Inc. All Rights Reserved.
+
+#include "AROriginActor.h"
+
+AAROriginActor::AAROriginActor(const FObjectInitializer& ObjectInitializer)
+	: Super(ObjectInitializer)
+{
+}
diff --git a/Engine/Source/Runtime/AugmentedReality/Public/AROriginActor.h b/Engine/Source/Runtime/AugmentedReality/Public/AROriginActor.h
new file mode 100644
--- /dev/null
+++ b/Engine/Source/Runtime/AugmentedReality/Public/AROriginActor.h
@@ -0,0 +1,19 @@
+// Copyright 1998-2019 Epic Games, Inc. All Rights Reserved.
+
+#pragma once
+
+#include "GameFramework/Info.h"
+
+#include "AROriginActor.generated.h"
+
+/**
+ * Simple actor that is spawned at the origin for AR systems that want to hang components on an actor
+ * Spawned as a custom class for easier TObjectIterator use
+ */
+UCLASS(BlueprintType)
+class AUGMENTEDREALITY_API AAROriginActor :
+	public AInfo
+{
+	GENERATED_UCLASS_BODY()
+public:
+};
