MERGE: Added a factory method to return the singleton AROriginActor to the requesting AR system

diff --git a/Engine/Source/Runtime/AugmentedReality/Private/AROriginActor.cpp b/Engine/Source/Runtime/AugmentedReality/Private/AROriginActor.cpp
--- a/Engine/Source/Runtime/AugmentedReality/Private/AROriginActor.cpp
+++ b/Engine/Source/Runtime/AugmentedReality/Private/AROriginActor.cpp
@@ -1,8 +1,25 @@
 // Copyright 1998-2019 Epic Games, Inc. All Rights Reserved.
 
 #include "AROriginActor.h"
+#include "EngineUtils.h"
 
 AAROriginActor::AAROriginActor(const FObjectInitializer& ObjectInitializer)
 	: Super(ObjectInitializer)
 {
 }
+
+AAROriginActor* AAROriginActor::GetOriginActor()
+{
+	if (GWorld != nullptr)
+	{
+		TActorIterator<AAROriginActor> Iter(GWorld);
+		AAROriginActor* FoundActor = *Iter;
+		if (FoundActor == nullptr)
+		{
+			// None spawned yet
+			FoundActor = GWorld->SpawnActor<AAROriginActor>(AAROriginActor::StaticClass(), FVector::ZeroVector, FRotator::ZeroRotator);
+		}
+		return FoundActor;
+	}
+	return nullptr;
+}
diff --git a/Engine/Source/Runtime/AugmentedReality/Public/AROriginActor.h b/Engine/Source/Runtime/AugmentedReality/Public/AROriginActor.h
--- a/Engine/Source/Runtime/AugmentedReality/Public/AROriginActor.h
+++ b/Engine/Source/Runtime/AugmentedReality/Public/AROriginActor.h
@@ -16,4 +16,6 @@ class AUGMENTEDREALITY_API AAROriginActor :
 {
 	GENERATED_UCLASS_BODY()
 public:
+	/** Used by the AR system to get the origin actor for the current world */
+	static AAROriginActor* GetOriginActor();
 };
