MERGE: Changed factory method

diff --git a/Engine/Source/Runtime/AugmentedReality/Private/AROriginActor.cpp b/Engine/Source/Runtime/AugmentedReality/Private/AROriginActor.cpp
--- a/Engine/Source/Runtime/AugmentedReality/Private/AROriginActor.cpp
+++ b/Engine/Source/Runtime/AugmentedReality/Private/AROriginActor.cpp
@@ -12,8 +12,15 @@ AAROriginActor* AAROriginActor::GetOriginActor()
 {
 	if (GWorld != nullptr)
 	{
-		TActorIterator<AAROriginActor> Iter(GWorld);
-		AAROriginActor* FoundActor = *Iter;
+		AAROriginActor* FoundActor = nullptr;
+		for (TActorIterator<AAROriginActor> Iter(GWorld); Iter; ++Iter)
+		{
+			if (!(*Iter)->IsPendingKill())
+			{
+				FoundActor = *Iter;
+				break;
+			}
+		}
 		if (FoundActor == nullptr)
 		{
 			// None spawned yet
