MERGE: Changed the lineage of AROriginActor to be straight from AActor

diff --git a/Engine/Source/Runtime/AugmentedReality/Private/AROriginActor.cpp b/Engine/Source/Runtime/AugmentedReality/Private/AROriginActor.cpp
--- a/Engine/Source/Runtime/AugmentedReality/Private/AROriginActor.cpp
+++ b/Engine/Source/Runtime/AugmentedReality/Private/AROriginActor.cpp
@@ -6,6 +6,11 @@
 AAROriginActor::AAROriginActor(const FObjectInitializer& ObjectInitializer)
 	: Super(ObjectInitializer)
 {
+	PrimaryActorTick.bCanEverTick = false;
+	bAllowTickBeforeBeginPlay = false;
+	bReplicates = false;
+	bReplicateMovement = false;
+	bCanBeDamaged = false;
 }
 
 AAROriginActor* AAROriginActor::GetOriginActor()
diff --git a/Engine/Source/Runtime/AugmentedReality/Public/AROriginActor.h b/Engine/Source/Runtime/AugmentedReality/Public/AROriginActor.h
--- a/Engine/Source/Runtime/AugmentedReality/Public/AROriginActor.h
+++ b/Engine/Source/Runtime/AugmentedReality/Public/AROriginActor.h
@@ -2,7 +2,7 @@
 
 #pragma once
 
-#include "GameFramework/Info.h"
+#include "GameframeWork/Actor.h"
 
 #include "AROriginActor.generated.h"
 
@@ -12,7 +12,7 @@
  */
 UCLASS(BlueprintType)
 class AUGMENTEDREALITY_API AAROriginActor :
-	public AInfo
+	public AActor
 {
 	GENERATED_UCLASS_BODY()
 public:
