MERGE: Added support for an underlying mesh to be part of UARTrackedGeometry

diff --git a/Engine/Source/Runtime/AugmentedReality/AugmentedReality.build.cs b/Engine/Source/Runtime/AugmentedReality/AugmentedReality.build.cs
--- a/Engine/Source/Runtime/AugmentedReality/AugmentedReality.build.cs
+++ b/Engine/Source/Runtime/AugmentedReality/AugmentedReality.build.cs
@@ -18,7 +18,7 @@ namespace UnrealBuildTool.Rules
                     "EngineSettings",
                     "RenderCore",
 					"RHI",
-                    //"HeadMountedDisplay"
+					"MRMesh"
 				}
 			);
 
diff --git a/Engine/Source/Runtime/AugmentedReality/Private/ARTrackable.cpp b/Engine/Source/Runtime/AugmentedReality/Private/ARTrackable.cpp
--- a/Engine/Source/Runtime/AugmentedReality/Private/ARTrackable.cpp
+++ b/Engine/Source/Runtime/AugmentedReality/Private/ARTrackable.cpp
@@ -4,6 +4,7 @@
 #include "ARSystem.h"
 #include "ARDebugDrawHelpers.h"
 #include "DrawDebugHelpers.h"
+#include "MRMeshComponent.h"
 
 //
 //
@@ -115,6 +116,16 @@ IARRef* UARTrackedGeometry::GetNativeResource()
 	return NativeResource.Get();
 }
 
+UMRMeshComponent* UARTrackedGeometry::GetUnderlyingMesh()
+{
+	return UnderlyingMesh;
+}
+
+void UARTrackedGeometry::SetUnderlyingMesh(UMRMeshComponent* InMRMeshComponent)
+{
+	UnderlyingMesh = InMRMeshComponent;
+}
+
 //
 //
 //
diff --git a/Engine/Source/Runtime/AugmentedReality/Public/ARTrackable.h b/Engine/Source/Runtime/AugmentedReality/Public/ARTrackable.h
--- a/Engine/Source/Runtime/AugmentedReality/Public/ARTrackable.h
+++ b/Engine/Source/Runtime/AugmentedReality/Public/ARTrackable.h
@@ -7,6 +7,7 @@
 
 class FARSupportInterface ;
 class UAREnvironmentCaptureProbeTexture;
+class UMRMeshComponent;
 
 UCLASS(BlueprintType)
 class AUGMENTEDREALITY_API UARTrackedGeometry : public UObject
@@ -25,7 +26,7 @@ public:
 	void UpdateTrackingState( EARTrackingState NewTrackingState );
 	
 	void UpdateAlignmentTransform( const FTransform& NewAlignmentTransform );
-	
+
 	void SetDebugName( FName InDebugName );
 
 	IARRef* GetNativeResource();
@@ -55,6 +56,10 @@ public:
 	float GetLastUpdateTimestamp() const;
 	inline void SetLastUpdateTimestamp(double InTimestamp) { LastUpdateTimestamp = InTimestamp; }
 
+	UFUNCTION(BlueprintPure, Category="AR AugmentedReality|Tracked Geometry")
+	UMRMeshComponent* GetUnderlyingMesh();
+	void SetUnderlyingMesh(UMRMeshComponent* InMRMeshComponent);
+
 	UPROPERTY(BlueprintReadOnly, Category="AR AugmentedReality|Tracked Geometry")
 	FGuid UniqueId;
 
@@ -73,6 +78,10 @@ protected:
 	/** A pointer to the native resource in the native AR system */
 	TUniquePtr<IARRef> NativeResource;
 	
+	/** For AR systems that support arbitrary mesh geometry associated with a tracked point */
+	UPROPERTY()
+	UMRMeshComponent* UnderlyingMesh;
+
 private:
 	TWeakPtr<FARSupportInterface , ESPMode::ThreadSafe> ARSystem;
 	
