Should be "GameFramework"

diff --git a/Engine/Source/Runtime/AugmentedReality/Public/AROriginActor.h b/Engine/Source/Runtime/AugmentedReality/Public/AROriginActor.h
--- a/Engine/Source/Runtime/AugmentedReality/Public/AROriginActor.h
+++ b/Engine/Source/Runtime/AugmentedReality/Public/AROriginActor.h
@@ -2,7 +2,7 @@
 
 #pragma once
 
-#include "GameframeWork/Actor.h"
+#include "GameFramework/Actor.h"
 
 #include "AROriginActor.generated.h"
 
