First pass at adding some scene understanding support to our AR layer

diff --git a/Engine/Source/Runtime/AugmentedReality/Private/ARTrackable.cpp b/Engine/Source/Runtime/AugmentedReality/Private/ARTrackable.cpp
--- a/Engine/Source/Runtime/AugmentedReality/Private/ARTrackable.cpp
+++ b/Engine/Source/Runtime/AugmentedReality/Private/ARTrackable.cpp
@@ -197,6 +197,7 @@ void UARTrackedImage::UpdateTrackedGeometry(const TSharedRef<FARSupportInterface
 	Super::UpdateTrackedGeometry(InTrackingSystem, FrameNumber, Timestamp, InLocalToTrackingTransform, InAlignmentTransform);
 	EstimatedSize = InEstimatedSize;
 	DetectedImage = InDetectedImage;
+	ObjectClassification = EARObjectClassification::Image;
 }
 
 FVector2D UARTrackedImage::GetEstimateSize()
diff --git a/Engine/Source/Runtime/AugmentedReality/Public/ARTrackable.h b/Engine/Source/Runtime/AugmentedReality/Public/ARTrackable.h
--- a/Engine/Source/Runtime/AugmentedReality/Public/ARTrackable.h
+++ b/Engine/Source/Runtime/AugmentedReality/Public/ARTrackable.h
@@ -127,7 +127,14 @@ public:
 	EARPlaneOrientation GetOrientation() const { return Orientation; }
 	void SetOrientation(EARPlaneOrientation InOrientation) { Orientation = InOrientation; }
 
+	UFUNCTION(BlueprintPure, Category="AR AugmentedReality|Plane Geometry")
+	EARObjectClassification GetObjectClassification() const { return ObjectClassification; }
+	void SetObjectClassification(EARObjectClassification InClassification) { ObjectClassification = InClassification; }
+
 private:
+	UPROPERTY()
+	EARObjectClassification ObjectClassification;
+
 	UPROPERTY()
 	EARPlaneOrientation Orientation;
 
@@ -186,6 +193,14 @@ public:
 	UPROPERTY(BlueprintReadOnly, Category="AR AugmentedReality|Face Geometry")
 	bool bIsTracked;
 
+	UFUNCTION(BlueprintPure, Category="AR AugmentedReality|Plane Geometry")
+	EARObjectClassification GetObjectClassification() const { return ObjectClassification; }
+	void SetObjectClassification(EARObjectClassification InClassification) { ObjectClassification = InClassification; }
+
+private:
+	UPROPERTY()
+	EARObjectClassification ObjectClassification;
+
 protected:
 	/** The candidate image that was detected in the scene */
 	UPROPERTY()
diff --git a/Engine/Source/Runtime/AugmentedReality/Public/ARTypes.h b/Engine/Source/Runtime/AugmentedReality/Public/ARTypes.h
--- a/Engine/Source/Runtime/AugmentedReality/Public/ARTypes.h
+++ b/Engine/Source/Runtime/AugmentedReality/Public/ARTypes.h
@@ -141,6 +141,27 @@ enum class EARPlaneOrientation : uint8
 	Vertical
 };
 
+/** Indicates what type of object the scene understanding system thinks it is */
+UENUM(BlueprintType)
+enum class EARObjectClassification : uint8
+{
+	/** Scene understanding doesn't know what this is */
+	Unknown,
+	/** The vertical plane that is a wall */
+	Wall,
+	/** The horizontal plane that is the ceiling */
+	Ceiling,
+	/** The horizontal plane that is the floor */
+	Floor,
+	/** The horizontal plane that is a table */
+	Table,
+	/** The horizontal plane that is a seat */
+	Seat,
+	/** A recognized image in the scene */
+	Image,
+	// Add other types here...
+};
+
 /** The current state of the AR subsystem including an optional explanation string. */
 USTRUCT(BlueprintType)
 struct AUGMENTEDREALITY_API FARSessionStatus
