Added a method to set the friendly name of a candidate object

diff --git a/Engine/Source/Runtime/AugmentedReality/Public/ARTypes.h b/Engine/Source/Runtime/AugmentedReality/Public/ARTypes.h
--- a/Engine/Source/Runtime/AugmentedReality/Public/ARTypes.h
+++ b/Engine/Source/Runtime/AugmentedReality/Public/ARTypes.h
@@ -304,6 +304,8 @@ public:
 	/** @see FriendlyName */
 	UFUNCTION(BlueprintPure, Category = "AR AugmentedReality|Object Detection")
 	const FString& GetFriendlyName() const { return FriendlyName; }
+	UFUNCTION(BlueprintCallable, Category = "AR AugmentedReality|Object Detection")
+	void SetFriendlyName(const FString& NewName) { FriendlyName = NewName; }
 
 	/** @see BoundingBox */
 	UFUNCTION(BlueprintPure, Category = "AR AugmentedReality|Object Detection")
