Additional scene understanding updates for AR

diff --git a/Engine/Source/Runtime/AugmentedReality/Private/ARTrackable.cpp b/Engine/Source/Runtime/AugmentedReality/Private/ARTrackable.cpp
--- a/Engine/Source/Runtime/AugmentedReality/Private/ARTrackable.cpp
+++ b/Engine/Source/Runtime/AugmentedReality/Private/ARTrackable.cpp
@@ -58,7 +58,7 @@ bool UARTrackedGeometry::IsTracked() const
 
 void UARTrackedGeometry::SetTrackingState(EARTrackingState NewState)
 {
-	TrackingState = NewState;
+	UpdateTrackingState(NewState);
 }
 
 FTransform UARTrackedGeometry::GetLocalToWorldTransform() const
@@ -94,10 +94,13 @@ void UARTrackedGeometry::UpdateTrackingState( EARTrackingState NewTrackingState
 {
 	TrackingState = NewTrackingState;
 
-	if (TrackingState == EARTrackingState::StoppedTracking && NativeResource)
+	if (TrackingState == EARTrackingState::StoppedTracking)
 	{
-		// Remove reference to the native resource since the tracked geometry is stopped tracking.
-		NativeResource->RemoveRef();
+		if (NativeResource.IsValid())
+		{
+			// Remove reference to the native resource since the tracked geometry is stopped tracking.
+			NativeResource->RemoveRef();
+		}
 	}
 }
 
@@ -219,6 +222,7 @@ void UARFaceGeometry::UpdateFaceGeometry(const TSharedRef<FARSupportInterface, E
 	LeftEyeTransform = InLeftEyeTransform;
 	RightEyeTransform = InRightEyeTransform;
 	LookAtTarget = InLookAtTarget;
+	ObjectClassification = EARObjectClassification::Face;
 }
 
 void UARTrackedPoint::DebugDraw(UWorld* World, const FLinearColor& OutlineColor, float OutlineThickness, float PersistForSeconds /*= 0.0f*/) const
@@ -322,5 +326,6 @@ void UARTrackedObject::UpdateTrackedGeometry(const TSharedRef<FARSupportInterfac
 {
 	Super::UpdateTrackedGeometry(InTrackingSystem, FrameNumber, Timestamp, InLocalToTrackingTransform, InAlignmentTransform);
 	DetectedObject = InDetectedObject;
+	ObjectClassification = EARObjectClassification::SceneObject;
 }
 
diff --git a/Engine/Source/Runtime/AugmentedReality/Public/ARTrackable.h b/Engine/Source/Runtime/AugmentedReality/Public/ARTrackable.h
--- a/Engine/Source/Runtime/AugmentedReality/Public/ARTrackable.h
+++ b/Engine/Source/Runtime/AugmentedReality/Public/ARTrackable.h
@@ -63,6 +63,10 @@ public:
 	UPROPERTY(BlueprintReadOnly, Category="AR AugmentedReality|Tracked Geometry")
 	FGuid UniqueId;
 
+	UFUNCTION(BlueprintPure, Category = "AR AugmentedReality|Scene Understanding")
+	EARObjectClassification GetObjectClassification() const { return ObjectClassification; }
+	void SetObjectClassification(EARObjectClassification InClassification) { ObjectClassification = InClassification; }
+
 protected:
 	TSharedPtr<FARSupportInterface , ESPMode::ThreadSafe> GetARSystem() const;
 	
@@ -82,6 +86,10 @@ protected:
 	UPROPERTY()
 	UMRMeshComponent* UnderlyingMesh;
 
+	/** What the scene understanding system thinks this object is */
+	UPROPERTY()
+	EARObjectClassification ObjectClassification;
+
 private:
 	TWeakPtr<FARSupportInterface , ESPMode::ThreadSafe> ARSystem;
 	
@@ -127,14 +135,7 @@ public:
 	EARPlaneOrientation GetOrientation() const { return Orientation; }
 	void SetOrientation(EARPlaneOrientation InOrientation) { Orientation = InOrientation; }
 
-	UFUNCTION(BlueprintPure, Category="AR AugmentedReality|Plane Geometry")
-	EARObjectClassification GetObjectClassification() const { return ObjectClassification; }
-	void SetObjectClassification(EARObjectClassification InClassification) { ObjectClassification = InClassification; }
-
 private:
-	UPROPERTY()
-	EARObjectClassification ObjectClassification;
-
 	UPROPERTY()
 	EARPlaneOrientation Orientation;
 
@@ -193,14 +194,6 @@ public:
 	UPROPERTY(BlueprintReadOnly, Category="AR AugmentedReality|Face Geometry")
 	bool bIsTracked;
 
-	UFUNCTION(BlueprintPure, Category="AR AugmentedReality|Plane Geometry")
-	EARObjectClassification GetObjectClassification() const { return ObjectClassification; }
-	void SetObjectClassification(EARObjectClassification InClassification) { ObjectClassification = InClassification; }
-
-private:
-	UPROPERTY()
-	EARObjectClassification ObjectClassification;
-
 protected:
 	/** The candidate image that was detected in the scene */
 	UPROPERTY()
diff --git a/Engine/Source/Runtime/AugmentedReality/Public/ARTypes.h b/Engine/Source/Runtime/AugmentedReality/Public/ARTypes.h
--- a/Engine/Source/Runtime/AugmentedReality/Public/ARTypes.h
+++ b/Engine/Source/Runtime/AugmentedReality/Public/ARTypes.h
@@ -145,20 +145,28 @@ enum class EARPlaneOrientation : uint8
 UENUM(BlueprintType)
 enum class EARObjectClassification : uint8
 {
+	/** Not applicable to scene understanding */
+	NotApplicable,
 	/** Scene understanding doesn't know what this is */
 	Unknown,
-	/** The vertical plane that is a wall */
+	/** A vertical plane that is a wall */
 	Wall,
-	/** The horizontal plane that is the ceiling */
+	/** A horizontal plane that is the ceiling */
 	Ceiling,
-	/** The horizontal plane that is the floor */
+	/** A horizontal plane that is the floor */
 	Floor,
-	/** The horizontal plane that is a table */
+	/** A horizontal plane that is a table */
 	Table,
-	/** The horizontal plane that is a seat */
+	/** A horizontal plane that is a seat */
 	Seat,
+	/** A human face */
+	Face,
 	/** A recognized image in the scene */
 	Image,
+	/** A chunk of mesh that does not map to a specific object type but is seen by the AR system */
+	World,
+	/** A mesh that was recognized in the scene */
+	SceneObject,
 	// Add other types here...
 };
 
