Added a AR session config option to enable/disable scene understanding detecting closed meshes in a scene

jira: UEVR-1393

diff --git a/Engine/Source/Runtime/AugmentedReality/Private/ARSessionConfig.cpp b/Engine/Source/Runtime/AugmentedReality/Private/ARSessionConfig.cpp
--- a/Engine/Source/Runtime/AugmentedReality/Private/ARSessionConfig.cpp
+++ b/Engine/Source/Runtime/AugmentedReality/Private/ARSessionConfig.cpp
@@ -7,19 +7,20 @@
 #include "ARSessionConfigCookSupport.h"
 
 UARSessionConfig::UARSessionConfig()
-: WorldAlignment(EARWorldAlignment::Gravity)
-, SessionType(EARSessionType::World)
-, PlaneDetectionMode_DEPRECATED(EARPlaneDetectionMode::HorizontalPlaneDetection)
-, bHorizontalPlaneDetection(true)
-, bVerticalPlaneDetection(true)
-, bEnableAutoFocus(true)
-, LightEstimationMode(EARLightEstimationMode::AmbientLightEstimate)
-, FrameSyncMode(EARFrameSyncMode::SyncTickWithoutCameraImage)
-, bEnableAutomaticCameraOverlay(true)
-, bEnableAutomaticCameraTracking(true)
-, bResetCameraTracking(true)
-, bResetTrackedObjects(true)
-, MaxNumSimultaneousImagesTracked(1)
+	: bTrackSceneObjects(true)
+	, WorldAlignment(EARWorldAlignment::Gravity)
+	, SessionType(EARSessionType::World)
+	, PlaneDetectionMode_DEPRECATED(EARPlaneDetectionMode::HorizontalPlaneDetection)
+	, bHorizontalPlaneDetection(true)
+	, bVerticalPlaneDetection(true)
+	, bEnableAutoFocus(true)
+	, LightEstimationMode(EARLightEstimationMode::AmbientLightEstimate)
+	, FrameSyncMode(EARFrameSyncMode::SyncTickWithoutCameraImage)
+	, bEnableAutomaticCameraOverlay(true)
+	, bEnableAutomaticCameraTracking(true)
+	, bResetCameraTracking(true)
+	, bResetTrackedObjects(true)
+	, MaxNumSimultaneousImagesTracked(1)
 {
 }
 
diff --git a/Engine/Source/Runtime/AugmentedReality/Public/ARSessionConfig.h b/Engine/Source/Runtime/AugmentedReality/Public/ARSessionConfig.h
--- a/Engine/Source/Runtime/AugmentedReality/Public/ARSessionConfig.h
+++ b/Engine/Source/Runtime/AugmentedReality/Public/ARSessionConfig.h
@@ -252,6 +252,10 @@ public:
 	UPROPERTY(EditAnywhere, Category = "AR Settings | World Mapping")
 	bool bRenderMeshDataInWireframe;
 
+	/** Whether the AR system should report scene objects (@see EARObjectClassification::SceneObject) */
+	UPROPERTY(EditAnywhere, Category = "AR Settings | World Mapping")
+	bool bTrackSceneObjects;
+
 private:
 	//~ UObject interface
 	virtual void Serialize(FArchive& Ar) override;
