Added support for AR systems that can match planes to slopes

diff --git a/Engine/Source/Runtime/AugmentedReality/Public/ARTypes.h b/Engine/Source/Runtime/AugmentedReality/Public/ARTypes.h
--- a/Engine/Source/Runtime/AugmentedReality/Public/ARTypes.h
+++ b/Engine/Source/Runtime/AugmentedReality/Public/ARTypes.h
@@ -138,7 +138,9 @@ UENUM(BlueprintType)
 enum class EARPlaneOrientation : uint8
 {
 	Horizontal,
-	Vertical
+	Vertical,
+	/** For AR systems that can match planes to slopes */
+	Diagonal,
 };
 
 /** Indicates what type of object the scene understanding system thinks it is */
@@ -165,7 +167,7 @@ enum class EARObjectClassification : uint8
 	Image,
 	/** A chunk of mesh that does not map to a specific object type but is seen by the AR system */
 	World,
-	/** A mesh that was recognized in the scene */
+	/** A closed mesh that was identified in the scene */
 	SceneObject,
 	// Add other types here...
 };
