Merged from partner branch:

Fixed the FaceAR sample

- Marked UARFaceGeometry::bIsTracked as deprecated in blueprint and changed FaceARPawn to use the correct API instead
- Set LivelinkTrackingType to FaceTracking to enable live link for Face AR

diff --git a/Engine/Source/Runtime/AugmentedReality/Public/ARTrackable.h b/Engine/Source/Runtime/AugmentedReality/Public/ARTrackable.h
--- a/Engine/Source/Runtime/AugmentedReality/Public/ARTrackable.h
+++ b/Engine/Source/Runtime/AugmentedReality/Public/ARTrackable.h
@@ -346,7 +346,7 @@ public:
 
 	UE_DEPRECATED(4.21, "This property is now deprecated, please use GetTrackingState() and check for EARTrackingState::Tracking or IsTracked() instead.")
 	/** Whether the face is currently being tracked by the AR system */
-	UPROPERTY(BlueprintReadOnly, Category="AR AugmentedReality|Face Geometry")
+	UPROPERTY(BlueprintReadOnly, Category="AR AugmentedReality|Face Geometry", meta=(Deprecated, DeprecationMessage = "This property is now deprecated, please use GetTrackingState() and check for EARTrackingState::Tracking or IsTracked() instead."))
 	bool bIsTracked;
 
 private:
