Fix ARKit build by providing a float version of CalculateVertexNormals.

diff --git a/Engine/Source/Runtime/AugmentedReality/Private/ARComponent.cpp b/Engine/Source/Runtime/AugmentedReality/Private/ARComponent.cpp
--- a/Engine/Source/Runtime/AugmentedReality/Private/ARComponent.cpp
+++ b/Engine/Source/Runtime/AugmentedReality/Private/ARComponent.cpp
@@ -575,6 +575,18 @@ void UARMeshComponent::UpdateVisualization_Implementation()
 }
 
 void FAccumulatedNormal::CalculateVertexNormals(TArray<FAccumulatedNormal>& AccumulatedNormals, const TArray<FVector>& Vertices, const TArray<MRMESH_INDEX_TYPE>& Indices, TArray<FPackedNormal>& OutTangentData, FVector MeshCenter, float PositionScale)
+{
+	// TODO: Revisit LWC and AR geometry; temp fix to unbreak build.
+	TArray<FVector3f> FloatVertices;
+	FloatVertices.SetNumUninitialized(Vertices.Num());
+	for (int32 i = 0; i < Vertices.Num(); ++i)
+	{
+		FloatVertices[i] = static_cast<FVector3f>(Vertices[i]);
+	}
+	CalculateVertexNormals(AccumulatedNormals, FloatVertices, Indices, OutTangentData, MeshCenter, PositionScale);
+}
+
+void FAccumulatedNormal::CalculateVertexNormals(TArray<FAccumulatedNormal>& AccumulatedNormals, const TArray<FVector3f>& Vertices, const TArray<MRMESH_INDEX_TYPE>& Indices, TArray<FPackedNormal>& OutTangentData, FVector MeshCenter, float PositionScale)
 {
 	SCOPE_CYCLE_COUNTER(STAT_CalculateVertexNormals);
 	
diff --git a/Engine/Source/Runtime/AugmentedReality/Public/ARComponent.h b/Engine/Source/Runtime/AugmentedReality/Public/ARComponent.h
--- a/Engine/Source/Runtime/AugmentedReality/Public/ARComponent.h
+++ b/Engine/Source/Runtime/AugmentedReality/Public/ARComponent.h
@@ -211,6 +211,7 @@ struct AUGMENTEDREALITY_API FAccumulatedNormal
 	uint32 NumFaces = 0;
 	
 	static void CalculateVertexNormals(TArray<FAccumulatedNormal>& AccumulatedNormals, const TArray<FVector>& Vertices, const TArray<MRMESH_INDEX_TYPE>& Indices, TArray<FPackedNormal>& OutTangentData, FVector MeshCenter = FVector::ZeroVector, float PositionScale = 1.f);
+	static void CalculateVertexNormals(TArray<FAccumulatedNormal>& AccumulatedNormals, const TArray<FVector3f>& Vertices, const TArray<MRMESH_INDEX_TYPE>& Indices, TArray<FPackedNormal>& OutTangentData, FVector MeshCenter = FVector::ZeroVector, float PositionScale = 1.f);
 };
 
 /**
