Support FAccumulatedNormal::CalculateVertexNormals for both float and double vector Vertices arrays, CIS blocker

diff --git a/Engine/Source/Runtime/AugmentedReality/Private/ARComponent.cpp b/Engine/Source/Runtime/AugmentedReality/Private/ARComponent.cpp
--- a/Engine/Source/Runtime/AugmentedReality/Private/ARComponent.cpp
+++ b/Engine/Source/Runtime/AugmentedReality/Private/ARComponent.cpp
@@ -574,7 +574,8 @@ void UARMeshComponent::UpdateVisualization_Implementation()
 	}
 }
 
-void FAccumulatedNormal::CalculateVertexNormals(TArray<FAccumulatedNormal>& AccumulatedNormals, const TArray<FVector>& Vertices, const TArray<MRMESH_INDEX_TYPE>& Indices, TArray<FPackedNormal>& OutTangentData, FVector MeshCenter, float PositionScale)
+template<typename VertexType>
+void FAccumulatedNormal::CalculateVertexNormals(TArray<FAccumulatedNormal>& AccumulatedNormals, const TArray<VertexType>& Vertices, const TArray<MRMESH_INDEX_TYPE>& Indices, TArray<FPackedNormal>& OutTangentData, FVector MeshCenter, float PositionScale)
 {
 	SCOPE_CYCLE_COUNTER(STAT_CalculateVertexNormals);
 	
@@ -632,6 +633,9 @@ void FAccumulatedNormal::CalculateVertexNormals(TArray<FAccumulatedNormal>& Accu
 		OutTangentData[2 * Index + 1] = FPackedNormal(AccumulatedNormal.Normal);
 	}
 }
+// instantiate for FVector3d and FVector3f
+template void FAccumulatedNormal::CalculateVertexNormals<FVector3d>(TArray<FAccumulatedNormal>& AccumulatedNormals, const TArray<FVector>& Vertices, const TArray<MRMESH_INDEX_TYPE>& Indices, TArray<FPackedNormal>& OutTangentData, FVector MeshCenter, float PositionScale);
+template void FAccumulatedNormal::CalculateVertexNormals<FVector3f>(TArray<FAccumulatedNormal>& AccumulatedNormals, const TArray<FVector3f>& Vertices, const TArray<MRMESH_INDEX_TYPE>& Indices, TArray<FPackedNormal>& OutTangentData, FVector MeshCenter, float PositionScale);
 
 void UARFaceComponent::UpdateVisualization_Implementation()
 {
diff --git a/Engine/Source/Runtime/AugmentedReality/Public/ARComponent.h b/Engine/Source/Runtime/AugmentedReality/Public/ARComponent.h
--- a/Engine/Source/Runtime/AugmentedReality/Public/ARComponent.h
+++ b/Engine/Source/Runtime/AugmentedReality/Public/ARComponent.h
@@ -210,7 +210,8 @@ struct AUGMENTEDREALITY_API FAccumulatedNormal
 	FVector Normal = FVector::ZeroVector;
 	uint32 NumFaces = 0;
 	
-	static void CalculateVertexNormals(TArray<FAccumulatedNormal>& AccumulatedNormals, const TArray<FVector>& Vertices, const TArray<MRMESH_INDEX_TYPE>& Indices, TArray<FPackedNormal>& OutTangentData, FVector MeshCenter = FVector::ZeroVector, float PositionScale = 1.f);
+	template<typename VertexType>
+	static void CalculateVertexNormals(TArray<FAccumulatedNormal>& AccumulatedNormals, const TArray<VertexType>& Vertices, const TArray<MRMESH_INDEX_TYPE>& Indices, TArray<FPackedNormal>& OutTangentData, FVector MeshCenter = FVector::ZeroVector, float PositionScale = 1.f);
 };
 
 /**
