Fix typo in CL 18595201 breaking CIS
none

diff --git a/Engine/Source/Runtime/AugmentedReality/Private/ARComponent.cpp b/Engine/Source/Runtime/AugmentedReality/Private/ARComponent.cpp
--- a/Engine/Source/Runtime/AugmentedReality/Private/ARComponent.cpp
+++ b/Engine/Source/Runtime/AugmentedReality/Private/ARComponent.cpp
@@ -634,7 +634,7 @@ void FAccumulatedNormal::CalculateVertexNormals(TArray<FAccumulatedNormal>& Accu
 	}
 }
 // instantiate for FVector3d and FVector3f
-template void FAccumulatedNormal::CalculateVertexNormals<FVector3d>(TArray<FAccumulatedNormal>& AccumulatedNormals, const TArray<FVector>& Vertices, const TArray<MRMESH_INDEX_TYPE>& Indices, TArray<FPackedNormal>& OutTangentData, FVector MeshCenter, float PositionScale);
+template void FAccumulatedNormal::CalculateVertexNormals<FVector3d>(TArray<FAccumulatedNormal>& AccumulatedNormals, const TArray<FVector3d>& Vertices, const TArray<MRMESH_INDEX_TYPE>& Indices, TArray<FPackedNormal>& OutTangentData, FVector MeshCenter, float PositionScale);
 template void FAccumulatedNormal::CalculateVertexNormals<FVector3f>(TArray<FAccumulatedNormal>& AccumulatedNormals, const TArray<FVector3f>& Vertices, const TArray<MRMESH_INDEX_TYPE>& Indices, TArray<FPackedNormal>& OutTangentData, FVector MeshCenter, float PositionScale);
 
 void UARFaceComponent::UpdateVisualization_Implementation()
