Fix potential buffer overflow condition.

diff --git a/Engine/Source/Runtime/AugmentedReality/Private/ARSharedWorldGameState.cpp b/Engine/Source/Runtime/AugmentedReality/Private/ARSharedWorldGameState.cpp
--- a/Engine/Source/Runtime/AugmentedReality/Private/ARSharedWorldGameState.cpp
+++ b/Engine/Source/Runtime/AugmentedReality/Private/ARSharedWorldGameState.cpp
@@ -54,7 +54,8 @@ void AARSharedWorldGameState::UpdatePreviewImageData(int32 Offset, const uint8*
 		return;
 	}
 	
-	if (Offset >= 0 && (Offset + Size) <= PreviewImageBytesTotal)
+	// Validate offset, and also check for signed integer overflow case.
+	if (Offset >= 0 && (Offset + Size) <= PreviewImageBytesTotal && (Offset + Size) >= 0)
 	{
 		uint8* PreviewImageBuffer = PreviewImageData.GetData();
 		FMemory::Memcpy((void*)&PreviewImageBuffer[Offset], (void*)Buffer, Size);
@@ -78,7 +79,8 @@ void AARSharedWorldGameState::UpdateARWorldData(int32 Offset, const uint8* Buffe
 		return;
 	}
 	
-	if (Offset >= 0 && (Offset + Size) <= ARWorldBytesTotal)
+	// Validate offset, and also check for signed integer overflow case.
+	if (Offset >= 0 && (Offset + Size) <= ARWorldBytesTotal && (Offset + Size) >= 0)
 	{
 		uint8* ARWorldBuffer = ARWorldData.GetData();
 		FMemory::Memcpy((void*)&ARWorldBuffer[Offset], (void*)Buffer, Size);
diff --git a/Engine/Source/Runtime/AugmentedReality/Private/ARSharedWorldPlayerController.cpp b/Engine/Source/Runtime/AugmentedReality/Private/ARSharedWorldPlayerController.cpp
--- a/Engine/Source/Runtime/AugmentedReality/Private/ARSharedWorldPlayerController.cpp
+++ b/Engine/Source/Runtime/AugmentedReality/Private/ARSharedWorldPlayerController.cpp
@@ -65,7 +65,8 @@ void AARSharedWorldPlayerController::ClientInitSharedWorld_Implementation(int32
 
 bool AARSharedWorldPlayerController::ClientUpdatePreviewImageData_Validate(int32 Offset, const TArray<uint8>& Buffer)
 {
-	return Offset >= 0;
+	// Validate offset, and also check for signed integer overflow case.
+	return Offset >= 0 && (Offset + Buffer.Num()) >= 0;
 }
 
 void AARSharedWorldPlayerController::ClientUpdatePreviewImageData_Implementation(int32 Offset, const TArray<uint8>& Buffer)
@@ -77,7 +78,8 @@ void AARSharedWorldPlayerController::ClientUpdatePreviewImageData_Implementation
 
 bool AARSharedWorldPlayerController::ClientUpdateARWorldData_Validate(int32 Offset, const TArray<uint8>& Buffer)
 {
-	return Offset >= 0;
+	// Validate offset, and also check for signed integer overflow case.
+	return Offset >= 0 && (Offset + Buffer.Num()) >= 0;
 }
 
 void AARSharedWorldPlayerController::ClientUpdateARWorldData_Implementation(int32 Offset, const TArray<uint8>& Buffer)
