ARBlueprintLibrary Deprecated function deletion
-Deleted some functions deprecated in 4.26 and 4.22.

diff --git a/Engine/Plugins/Runtime/AR/Google/GoogleARCore/Source/GoogleARCoreBase/Private/GoogleARCoreAugmentedImage.cpp b/Engine/Plugins/Runtime/AR/Google/GoogleARCore/Source/GoogleARCoreBase/Private/GoogleARCoreAugmentedImage.cpp
--- a/Engine/Plugins/Runtime/AR/Google/GoogleARCore/Source/GoogleARCoreBase/Private/GoogleARCoreAugmentedImage.cpp
+++ b/Engine/Plugins/Runtime/AR/Google/GoogleARCore/Source/GoogleARCoreBase/Private/GoogleARCoreAugmentedImage.cpp
@@ -79,31 +79,3 @@ void UGoogleARCoreAugmentedImage::UpdateTrackedGeometry(
 	ImageIndex = InImageIndex;
 	ImageName = InImageName;
 }
-
-FVector UGoogleARCoreAugmentedImage::GetCenter() const
-{
-	return FVector::ZeroVector;
-}
-
-FVector UGoogleARCoreAugmentedImage::GetExtent() const
-{
-	return FVector(EstimatedSize.Y, 0.0f, EstimatedSize.X);
-}
-
-int32 UGoogleARCoreAugmentedImage::GetImageIndex() const
-{
-	return ImageIndex;
-}
-
-FString UGoogleARCoreAugmentedImage::GetImageName() const
-{
-	if (DetectedImage != nullptr)
-	{
-		return GetDetectedImage()->GetFriendlyName();
-
-	}
-	else
-	{
-		return ImageName;
-	}
-}
diff --git a/Engine/Plugins/Runtime/AR/Google/GoogleARCore/Source/GoogleARCoreBase/Private/GoogleARCoreAugmentedImageDatabase.cpp b/Engine/Plugins/Runtime/AR/Google/GoogleARCore/Source/GoogleARCoreBase/Private/GoogleARCoreAugmentedImageDatabase.cpp
--- a/Engine/Plugins/Runtime/AR/Google/GoogleARCore/Source/GoogleARCoreBase/Private/GoogleARCoreAugmentedImageDatabase.cpp
+++ b/Engine/Plugins/Runtime/AR/Google/GoogleARCore/Source/GoogleARCoreBase/Private/GoogleARCoreAugmentedImageDatabase.cpp
@@ -19,74 +19,6 @@
 #include "png.h"
 #include "TextureResource.h"
 
-int UGoogleARCoreAugmentedImageDatabase::AddRuntimeAugmentedImageFromTexture(UTexture2D* ImageTexture, FName ImageName, float ImageWidthInMeter /*= 0*/)
-{
-	EPixelFormat PixelFormat = ImageTexture->GetPixelFormat();
-
-	if (PixelFormat == EPixelFormat::PF_B8G8R8A8 || PixelFormat == EPixelFormat::PF_G8)
-	{
-		ensure(ImageTexture->GetNumMips() > 0);
-		FTexture2DMipMap* Mip0 = &ImageTexture->GetPlatformData()->Mips[0];
-		FByteBulkData* RawImageData = &Mip0->BulkData;
-
-		int ImageWidth = ImageTexture->GetSizeX();
-		int ImageHeight = ImageTexture->GetSizeY();
-
-		TArray<uint8> GrayscaleBuffer;
-		int PixelNum = ImageWidth * ImageHeight;
-		uint8* RawBytes = static_cast<uint8*>(RawImageData->Lock(LOCK_READ_ONLY));
-		if (PixelFormat == EPixelFormat::PF_B8G8R8A8)
-		{
-			GrayscaleBuffer.SetNumUninitialized(PixelNum);
-			ensureMsgf(RawImageData->GetBulkDataSize() == ImageWidth * ImageHeight * 4,
-				TEXT("Unsupported texture data in UGoogleARCoreAugmentedImageDatabase::AddRuntimeAugmentedImage"));
-
-			for (int i = 0; i < PixelNum; i++)
-			{
-				uint8 R = RawBytes[i * 4 + 2];
-				uint8 G = RawBytes[i * 4 + 1];
-				uint8 B = RawBytes[i * 4];
-				GrayscaleBuffer[i] = 0.2126 * R + 0.7152 * G + 0.0722 * B;
-			}
-		}
-		else
-		{
-			ensureMsgf(RawImageData->GetBulkDataSize() == ImageWidth * ImageHeight,
-				TEXT("Unsupported texture data in UGoogleARCoreAugmentedImageDatabase::AddRuntimeAugmentedImage"));
-			GrayscaleBuffer = TArray<uint8>(RawBytes, PixelNum);
-		}
-		RawImageData->Unlock();
-
-PRAGMA_DISABLE_DEPRECATION_WARNINGS
-		return AddRuntimeAugmentedImage(GrayscaleBuffer, ImageWidth, ImageHeight, ImageName, ImageWidthInMeter, ImageTexture);
-PRAGMA_ENABLE_DEPRECATION_WARNINGS
-	}
-
-	UE_LOG(LogGoogleARCore, Warning, TEXT("Failed to add runtime augmented image: Unsupported texture format: %s. ARCore only support PF_B8G8R8A8 or PF_G8 for now for adding runtime Augmented Image"), GetPixelFormatString(PixelFormat));
-	return -1;
-}
-
-int UGoogleARCoreAugmentedImageDatabase::AddRuntimeAugmentedImage(const TArray<uint8>& ImageGrayscalePixels,
-	int ImageWidth, int ImageHeight, FName ImageName, float ImageWidthInMeter /*= 0*/, UTexture2D* ImageTexture /*= nullptr*/)
-{
-	int NewImageIndex = FGoogleARCoreDevice::GetInstance()->AddRuntimeAugmentedImage(this, ImageGrayscalePixels,
-		ImageWidth, ImageHeight, ImageName.ToString(), ImageWidthInMeter);
-
-	if (NewImageIndex == -1)
-	{
-		return -1;
-	}
-
-	FGoogleARCoreAugmentedImageDatabaseEntry NewEntry;
-	NewEntry.Name = ImageName;
-	NewEntry.ImageAsset = ImageTexture;
-	NewEntry.Width = ImageWidthInMeter;
-
-	Entries.Add(NewEntry);
-
-	return NewImageIndex;
-}
-
 void UGoogleARCoreAugmentedImageDatabase::Serialize(FArchive& Ar)
 {
 	FString OutStdout;
diff --git a/Engine/Plugins/Runtime/AR/Google/GoogleARCore/Source/GoogleARCoreBase/Public/GoogleARCoreAugmentedImage.h b/Engine/Plugins/Runtime/AR/Google/GoogleARCore/Source/GoogleARCoreBase/Public/GoogleARCoreAugmentedImage.h
--- a/Engine/Plugins/Runtime/AR/Google/GoogleARCore/Source/GoogleARCoreBase/Public/GoogleARCoreAugmentedImage.h
+++ b/Engine/Plugins/Runtime/AR/Google/GoogleARCore/Source/GoogleARCoreBase/Public/GoogleARCoreAugmentedImage.h
@@ -18,47 +18,6 @@ class GOOGLEARCOREBASE_API UGoogleARCoreAugmentedImage : public UARTrackedImage
 
 public:
 
-	/**
-	 * Get the center in local space of the augmented image.
-	 *
-	 * @return The center in the augmented image's local coordinate
-	 *         space.
-	 */
-	UE_DEPRECATED(4.22, "There is no need to use this function since it always returns (0,0,0).")
-	UFUNCTION(BlueprintPure, Category = "GoogleARCore|AugmentedImage", meta = (Keywords = "googlear arcore augmentedimage"))
-	FVector GetCenter() const;
-
-	/**
-	 * Get the size in local space of the augmented image.
-	 *
-	 * @return The size in the augmented image's local coordinate
-	 *         space.
-	 */
-	UE_DEPRECATED(4.22, "Please use UARTrackedImage::GetEstimatedSize() instead.")
-	UFUNCTION(BlueprintPure, Category = "GoogleARCore|AugmentedImage", meta = (Keywords = "googlear arcore augmentedimage"))
-	FVector GetExtent() const;
-
-	/**
-	 * Get the index of an augmented image in the augmented image
-	 * database.
-	 *
-	 * @return An integer indicating the index into the Entries array
-	 *         in the UGoogleARCoreAugmentedImageDatabase.
-	 */
-	UE_DEPRECATED(4.22, "Instead of getting the index, you can use UARTrackedImage::GetDetectedImage() to get the UARCandidateImage object.")
-	UFUNCTION(BlueprintPure, Category = "GoogleARCore|AugmentedImage", meta = (Keywords = "googlear arcore augmentedimage"))
-	int32 GetImageIndex() const;
-
-	/**
-	 * Get the name of an augmented image in the augmented image
-	 * database.
-	 *
-	 * @return The name of the image in the image database.
-	 */
-	UE_DEPRECATED(4.22, "Please use UARCandidateImage::GetFriendlyName() instead.")
-	UFUNCTION(BlueprintPure, Category = "GoogleARCore|AugmentedImage", meta = (Keywords = "googlear arcore augmentedimage"))
-	FString GetImageName() const;
-
 	/**
 	 * Draw a box around the image, for debugging purposes.
 	 *
diff --git a/Engine/Plugins/Runtime/AR/Google/GoogleARCore/Source/GoogleARCoreBase/Public/GoogleARCoreAugmentedImageDatabase.h b/Engine/Plugins/Runtime/AR/Google/GoogleARCore/Source/GoogleARCoreBase/Public/GoogleARCoreAugmentedImageDatabase.h
--- a/Engine/Plugins/Runtime/AR/Google/GoogleARCore/Source/GoogleARCoreBase/Public/GoogleARCoreAugmentedImageDatabase.h
+++ b/Engine/Plugins/Runtime/AR/Google/GoogleARCore/Source/GoogleARCoreBase/Public/GoogleARCoreAugmentedImageDatabase.h
@@ -62,45 +62,6 @@ class GOOGLEARCOREBASE_API UGoogleARCoreAugmentedImageDatabase : public UDataAss
 	friend class FGoogleARCoreSession;
 
 public:
-	/**
-	 * Adds a single named image from the a UTexture2D to this AugmentedImageDatabase. This function will
-	 * add a FGoogleARCoreAugmentedImageDatabaseEntry to the Entries property of this AugmentedImageDatabase.
-	 *
-	 * You need to restart the ARCore session with the config that contains this AugmentedImageDatabase to make
-	 * track those new images.
-	 *
-	 * You can set the ImageWidthInMeter value if the physical size of the image is known. This will
-	 * help ARCore estimate the pose of the physical image as soon as ARCore detects the physical image.
-	 * Otherwise, ARCore will requiring the user to move the device to view the physical image from
-	 * different viewports.
-
-	 * Note that this function takes time to perform non-trivial image processing (20ms - 30ms), and should
-	 * be run on a background thread.
-	 */
-	UE_DEPRECATED(4.22, "Please use UARBlueprintLibrary::AddRuntimeCandidateImage() instead.")
-	UFUNCTION(BlueprintCallable, Category = "google arcore augmentedimages")
-	int AddRuntimeAugmentedImageFromTexture(UTexture2D* ImageTexture, FName ImageName, float ImageWidthInMeter = 0);
-
-	/**
-	 * Adds a single named image from raw byte to this AugmentedImage database. This function will
-	 * add a FGoogleARCoreAugmentedImageDatabaseEntry to the Entries property of this AugmentedImageDatabase.
-	 * The ImageAsset property in FGoogleARCoreAugmentedImageDatabaseEntry will be null.
-	 *
-	 * You need to restart the ARCore session with the config that contains this AugmentedImageDatabase to make
-	 * track those new images.
-	 *
-	 * You can set the ImageWidthInMeter value if the physical size of the image is known. This will
-	 * help ARCore estimate the pose of the physical image as soon as ARCore detects the physical image.
-	 * Otherwise, ARCore will requiring the user to move the device to view the physical image from
-	 * different viewports.
-	 *
-	 * Note that this function takes time to perform non-trivial image processing (20ms - 30ms), and should
-	 * be run on a background thread.
-	 */
-	UE_DEPRECATED(4.22, "Please use UGoogleARCoreSessionFunctionLibrary::AddRuntimeCandidateImageFromRawbytes() instead.")
-	int AddRuntimeAugmentedImage(const TArray<uint8>& ImageGrayscalePixels, int ImageWidth, int ImageHeight,
-		FName ImageName, float ImageWidthInMeter = 0, UTexture2D* ImageTexture = nullptr);
-
 	/**
 	 * Overridden serialization function.
 	 */
diff --git a/Engine/Source/Runtime/AugmentedReality/Private/ARBlueprintLibrary.cpp b/Engine/Source/Runtime/AugmentedReality/Private/ARBlueprintLibrary.cpp
--- a/Engine/Source/Runtime/AugmentedReality/Private/ARBlueprintLibrary.cpp
+++ b/Engine/Source/Runtime/AugmentedReality/Private/ARBlueprintLibrary.cpp
@@ -451,92 +451,6 @@ EARLineTraceChannels UARTraceResultLibrary::GetTraceChannel( const FARTraceResul
 	return TraceResult.GetTraceChannel();
 }
 
-UARTextureCameraImage* UARBlueprintLibrary::GetCameraImage()
-{
-	return Cast<UARTextureCameraImage>(GetARTexture(EARTextureType::CameraImage));
-}
-
-UARTextureCameraDepth* UARBlueprintLibrary::GetCameraDepth()
-{
-	return Cast<UARTextureCameraDepth>(GetARTexture(EARTextureType::CameraDepth));
-}
-
-TArray<UARPlaneGeometry*> UARBlueprintLibrary::GetAllTrackedPlanes()
-{
-	TArray<UARPlaneGeometry*> ResultList;
-	
-	auto ARSystem = GetARSystem();
-	if (ARSystem.IsValid())
-	{
-		auto PinnedARSystem = ARSystem.Pin();
-		for (UARTrackedGeometry* Geo : PinnedARSystem->GetAllTrackedGeometries())
-		{
-			if (UARPlaneGeometry* Item = Cast<UARPlaneGeometry>(Geo))
-			{
-				ResultList.Add(Item);
-			}
-		}
-	}
-	return ResultList;
-}
-
-TArray<UARTrackedPoint*> UARBlueprintLibrary::GetAllTrackedPoints()
-{
-	TArray<UARTrackedPoint*> ResultList;
-	
-	auto ARSystem = GetARSystem();
-	if (ARSystem.IsValid())
-	{
-		auto PinnedARSystem = ARSystem.Pin();
-		for (UARTrackedGeometry* Geo : PinnedARSystem->GetAllTrackedGeometries())
-		{
-			if (UARTrackedPoint* Item = Cast<UARTrackedPoint>(Geo))
-			{
-				ResultList.Add(Item);
-			}
-		}
-	}
-	return ResultList;
-}
-
-TArray<UARTrackedImage*> UARBlueprintLibrary::GetAllTrackedImages()
-{
-	TArray<UARTrackedImage*> ResultList;
-	
-	auto ARSystem = GetARSystem();
-	if (ARSystem.IsValid())
-	{
-		auto PinnedARSystem = ARSystem.Pin();
-		for (UARTrackedGeometry* Geo : PinnedARSystem->GetAllTrackedGeometries())
-		{
-			if (UARTrackedImage* Item = Cast<UARTrackedImage>(Geo))
-			{
-				ResultList.Add(Item);
-			}
-		}
-	}
-	return ResultList;
-}
-
-TArray<UAREnvironmentCaptureProbe*> UARBlueprintLibrary::GetAllTrackedEnvironmentCaptureProbes()
-{
-	TArray<UAREnvironmentCaptureProbe*> ResultList;
-	
-	auto ARSystem = GetARSystem();
-	if (ARSystem.IsValid())
-	{
-		auto PinnedARSystem = ARSystem.Pin();
-		for (UARTrackedGeometry* Geo : PinnedARSystem->GetAllTrackedGeometries())
-		{
-			if (UAREnvironmentCaptureProbe* Item = Cast<UAREnvironmentCaptureProbe>(Geo))
-			{
-				ResultList.Add(Item);
-			}
-		}
-	}
-	return ResultList;
-}
-
 bool UARBlueprintLibrary::AddManualEnvironmentCaptureProbe(FVector Location, FVector Extent)
 {
 	auto ARSystem = GetARSystem();
@@ -638,35 +552,6 @@ TArray<FARPose2D> UARBlueprintLibrary::GetAllTracked2DPoses()
 	}
 }
 
-TArray<UARTrackedPose*> UARBlueprintLibrary::GetAllTrackedPoses()
-{
-	TArray<UARTrackedPose*> TrackedPoses;
-	
-	auto ARSystem = GetARSystem();
-	if (ARSystem.IsValid())
-	{
-		auto PinnedARSystem = ARSystem.Pin();
-		for (UARTrackedGeometry* Geo : PinnedARSystem->GetAllTrackedGeometries())
-		{
-			if (UARTrackedPose* Item = Cast<UARTrackedPose>(Geo))
-			{
-				TrackedPoses.Add(Item);
-			}
-		}
-	}
-	return TrackedPoses;
-}
-
-UARTexture* UARBlueprintLibrary::GetPersonSegmentationImage()
-{
-	return GetARTexture(EARTextureType::PersonSegmentationImage);
-}
-
-UARTexture* UARBlueprintLibrary::GetPersonSegmentationDepthImage()
-{
-	return GetARTexture(EARTextureType::PersonSegmentationDepth);
-}
-
 bool UARBlueprintLibrary::IsSceneReconstructionSupported(EARSessionType SessionType, EARSceneReconstruction SceneReconstructionMethod)
 {
 	auto ARSystem = GetARSystem();
diff --git a/Engine/Source/Runtime/AugmentedReality/Public/ARBlueprintLibrary.h b/Engine/Source/Runtime/AugmentedReality/Public/ARBlueprintLibrary.h
--- a/Engine/Source/Runtime/AugmentedReality/Public/ARBlueprintLibrary.h
+++ b/Engine/Source/Runtime/AugmentedReality/Public/ARBlueprintLibrary.h
@@ -154,16 +154,6 @@ public:
 	/** @return a list of all the real-world geometry of the specified class as currently seen by the Augmented Reality system */
 	UFUNCTION(BlueprintCallable, Category = "AR AugmentedReality|Tracking", meta = (DisplayName="Get All AR Geometries By Class", Keywords = "ar augmentedreality augmented reality tracking geometry anchor", DeterminesOutputType = "GeometryClass"))
 	static TArray<UARTrackedGeometry*> GetAllGeometriesByClass(TSubclassOf<UARTrackedGeometry> GeometryClass);
-
-	/** @return the current camera image from the Augmented Reality system */
-	UE_DEPRECATED(4.26, "GetCameraImage is deprecated, use GetARTexture.")
-	UFUNCTION(BlueprintCallable, Category = "AR AugmentedReality|Camera", meta = (DeprecatedFunction, DeprecationMessage="GetCameraImage is deprecated, use GetARTexture.", DisplayName="Get AR Camera Image", Keywords = "ar augmentedreality augmented reality camera image"))
-	static UARTextureCameraImage* GetCameraImage();
-
-	/** @return the current camera depth data from the Augmented Reality system */
-	UE_DEPRECATED(4.26, "GetCameraDepth is deprecated, use GetARTexture.")
-	UFUNCTION(BlueprintCallable, Category = "AR AugmentedReality|Camera", meta = (DeprecatedFunction, DeprecationMessage="GetCameraDepth is deprecated, use GetARTexture.", DisplayName="Get AR Camera Depth", Keywords = "ar augmentedreality augmented reality camera image depth"))
-	static UARTextureCameraDepth* GetCameraDepth();
 	
 	/** @return the AR texture for the specified type */
 	UFUNCTION(BlueprintCallable, Category = "AR AugmentedReality|Texture", meta = (DisplayName="Get AR Texture", Keywords = "ar augmentedreality augmented reality texture"))
@@ -294,26 +284,6 @@ public:
 	 */
 	UFUNCTION(BlueprintCallable, Category = "AR AugmentedReality|ARPin", meta = (Keywords = "ar augmentedreality augmented reality tracking arpin anchor LocalStore"))
 	static void RemoveAllARPinsFromLocalStore();
-	
-	/** @return a list of all the tracked planes as currently seen by the Augmented Reality system */
-	UE_DEPRECATED(4.26, "GetAllTrackedPlanes is deprecated, use GetAllGeometriesByClass.")
-	UFUNCTION(BlueprintCallable, Category = "AR AugmentedReality|Tracking", meta = (DisplayName="Get All AR Tracked Planes", Keywords = "ar augmentedreality augmented reality tracking geometry anchor", DeprecatedFunction, DeprecationMessage="GetAllTrackedPlanes is deprecated, use GetAllGeometriesByClass."))
-	static TArray<UARPlaneGeometry*> GetAllTrackedPlanes();
-	
-	/** @return a list of all the tracked points as currently seen by the Augmented Reality system */
-	UE_DEPRECATED(4.26, "GetAllTrackedPoints is deprecated, use GetAllGeometriesByClass.")
-	UFUNCTION(BlueprintCallable, Category = "AR AugmentedReality|Tracking", meta = (KDisplayName="Get All AR Tracked Points", eywords = "ar augmentedreality augmented reality tracking geometry anchor", DeprecatedFunction, DeprecationMessage="GetAllTrackedPoints is deprecated, use GetAllGeometriesByClass."))
-	static TArray<UARTrackedPoint*> GetAllTrackedPoints();
-	
-	/** @return a list of all the tracked images as currently seen by the Augmented Reality system */
-	UE_DEPRECATED(4.26, "GetAllTrackedImages is deprecated, use GetAllGeometriesByClass.")
-	UFUNCTION(BlueprintCallable, Category = "AR AugmentedReality|Tracking", meta = (DisplayName="Get All AR Tracked Images", Keywords = "ar augmentedreality augmented reality tracking images anchor", DeprecatedFunction, DeprecationMessage="GetAllTrackedImages is deprecated, use GetAllGeometriesByClass."))
-	static TArray<UARTrackedImage*> GetAllTrackedImages();
-	
-	/** @return a list of all the tracked environment capture probes as currently seen by the Augmented Reality system */
-	UE_DEPRECATED(4.26, "GetAllTrackedEnvironmentCaptureProbes is deprecated, use GetAllGeometriesByClass.")
-	UFUNCTION(BlueprintCallable, Category = "AR AugmentedReality|Tracking", meta = (DisplayName="Get All AR Tracked Environment Probes", Keywords = "ar augmentedreality augmented reality tracking anchor", DeprecatedFunction, DeprecationMessage="GetAllTrackedEnvironmentCaptureProbes is deprecated, use GetAllGeometriesByClass."))
-	static TArray<UAREnvironmentCaptureProbe*> GetAllTrackedEnvironmentCaptureProbes();
 
 	/** Adds an environment capture probe to the ar world */
 	UFUNCTION(BlueprintCallable, Category = "AR AugmentedReality|Tracking", meta = (DisplayName="Add AR Environment Probe", Keywords = "ar augmentedreality augmented reality tracking anchor"))
@@ -361,21 +331,6 @@ public:
 	UFUNCTION(BlueprintCallable, Category = "AR AugmentedReality|Pose Tracking", meta = (DisplayName="Get All AR Tracked 2D Poses", Keywords = "ar augmentedreality augmented reality pose tracking"))
 	static TArray<FARPose2D> GetAllTracked2DPoses();
 	
-	/** @return all the 3D poses tracked by the AR system */
-	UE_DEPRECATED(4.26, "GetAllTrackedPoses is deprecated, use GetAllGeometriesByClass.")
-	UFUNCTION(BlueprintCallable, Category = "AR AugmentedReality|Pose Tracking", meta = (DisplayName="Get All AR Tracked 3D Poses", Keywords = "ar augmentedreality augmented reality pose tracking", DeprecatedFunction, DeprecationMessage="GetAllTrackedPoses is deprecated, use GetAllGeometriesByClass."))
-	static TArray<UARTrackedPose*> GetAllTrackedPoses();
-	
-	/** @return the segmentation image if the person segmentation session feature is used */
-	UE_DEPRECATED(4.26, "GetPersonSegmentationImage is deprecated, use GetARTexture.")
-	UFUNCTION(BlueprintCallable, Category = "AR AugmentedReality|Person Segmentation", meta = (DeprecatedFunction, DeprecationMessage="GetPersonSegmentationImage is deprecated, use GetARTexture.", DisplayName="Get AR Person Segmentation Image", Keywords = "ar augmentedreality augmented reality person segmentation image"))
-	static UARTexture* GetPersonSegmentationImage();
-
-	/** @return the segmentation depth image if the person segmentation with depth session feature is used */
-	UE_DEPRECATED(4.26, "GetPersonSegmentationDepthImage is deprecated, use GetARTexture.")
-	UFUNCTION(BlueprintCallable, Category = "AR AugmentedReality|Person Segmentation", meta = (DeprecatedFunction, DeprecationMessage="GetPersonSegmentationDepthImage is deprecated, use GetARTexture.", DisplayName="Get AR Person Segmentation Depth Image", Keywords = "ar augmentedreality augmented reality person segmentation depth image"))
-	static UARTexture* GetPersonSegmentationDepthImage();
-	
 	/**
 	 * Try to determine the classification of the object at a world space location
 	 * @InWorldLocation: the world location where the classification is needed
