UE-154100 ARSessionConfig does not have a method to remove candidate images
-Added image removal methods.

diff --git a/Engine/Source/Runtime/AugmentedReality/Private/ARSessionConfig.cpp b/Engine/Source/Runtime/AugmentedReality/Private/ARSessionConfig.cpp
--- a/Engine/Source/Runtime/AugmentedReality/Private/ARSessionConfig.cpp
+++ b/Engine/Source/Runtime/AugmentedReality/Private/ARSessionConfig.cpp
@@ -119,6 +119,28 @@ void UARSessionConfig::AddCandidateImage(UARCandidateImage* NewCandidateImage)
 	CandidateImages.Add(NewCandidateImage);
 }
 
+void UARSessionConfig::RemoveCandidateImage(UARCandidateImage* CandidateImage)
+{
+	int ImagesRemoved = CandidateImages.Remove(CandidateImage);
+}
+
+void UARSessionConfig::RemoveCandidateImageAtIndex(int Index)
+{
+	if (Index < 0 || Index >= CandidateImages.Num())
+	{
+		UE_LOG(LogBlueprint, Warning, TEXT("RemoveCandidateImageAtIndex failed because the index is invalid.  No image removed."));
+	}
+	else
+	{
+		CandidateImages.RemoveAt(Index);
+	}
+}
+
+void UARSessionConfig::ClearCandidateImages()
+{
+	CandidateImages.Empty();
+}
+
 int32 UARSessionConfig::GetMaxNumSimultaneousImagesTracked() const
 {
     return MaxNumSimultaneousImagesTracked;
diff --git a/Engine/Source/Runtime/AugmentedReality/Public/ARSessionConfig.h b/Engine/Source/Runtime/AugmentedReality/Public/ARSessionConfig.h
--- a/Engine/Source/Runtime/AugmentedReality/Public/ARSessionConfig.h
+++ b/Engine/Source/Runtime/AugmentedReality/Public/ARSessionConfig.h
@@ -246,6 +246,18 @@ public:
 	UFUNCTION(BlueprintCallable, Category = "AR Settings")
 	void AddCandidateImage(UARCandidateImage* NewCandidateImage);
     
+	// Remove a candidate image from the ARSessionConfig, by pointer, note the image object must match, not the content of the image.
+	UFUNCTION(BlueprintCallable, Category = "AR Settings")
+	void RemoveCandidateImage(UARCandidateImage* CandidateImage);
+
+	// Remove a candidate image from the ARSessionConfig, by index.
+	UFUNCTION(BlueprintCallable, Category = "AR Settings")
+	void RemoveCandidateImageAtIndex(int Index);
+
+	// Remove all candidate images from the ARSessionConfig
+	UFUNCTION(BlueprintCallable, Category = "AR Settings")
+	void ClearCandidateImages();
+
 	/** @see MaxNumSimultaneousImagesTracked */
 	UFUNCTION(BlueprintCallable, Category = "AR Settings")
     int32 GetMaxNumSimultaneousImagesTracked() const;
