Updated headers using UnrealCodeFixup to make sure dllstorage is on methods/staticvars instead of types. Part 5/n

diff --git a/Engine/Source/Runtime/AugmentedReality/Public/ARBlueprintLibrary.h b/Engine/Source/Runtime/AugmentedReality/Public/ARBlueprintLibrary.h
--- a/Engine/Source/Runtime/AugmentedReality/Public/ARBlueprintLibrary.h
+++ b/Engine/Source/Runtime/AugmentedReality/Public/ARBlueprintLibrary.h
@@ -15,6 +15,8 @@
 
 #include "ARBlueprintLibrary.generated.h"
 
+#define UE_API AUGMENTEDREALITY_API
+
 
 #define DEFINE_AR_BPLIB_DELEGATE_FUNCS(DelegateName) \
 public: \
@@ -46,8 +48,8 @@ public: \
 		} \
 	}
 
-UCLASS(meta=(ScriptName="ARLibrary"))
-class AUGMENTEDREALITY_API UARBlueprintLibrary : public UBlueprintFunctionLibrary
+UCLASS(MinimalAPI, meta=(ScriptName="ARLibrary"))
+class UARBlueprintLibrary : public UBlueprintFunctionLibrary
 {
 	GENERATED_BODY()
 
@@ -58,7 +60,7 @@ public:
 	 *
 	 */
 	UFUNCTION(BlueprintCallable, Category = "AR AugmentedReality|Session", meta = (DisplayName="Is AR Supported", Keywords = "ar augmentedreality augmented reality session start run running"))
-	static bool IsARSupported();
+	static UE_API bool IsARSupported();
 	
 	/**
 	 * Begin a new Augmented Reality session. Subsequently, use the \c GetARSessionStatus() function to figure out the status of the session.
@@ -66,15 +68,15 @@ public:
 	 * @param SessionConfig    Describes the tracking method to use, what kind of geometry to detect in the world, etc.
 	 */
 	UFUNCTION(BlueprintCallable, Category = "AR AugmentedReality|Session", meta = (DisplayName="Start AR Session", Keywords = "ar augmentedreality augmented reality session start run running"))
-	static void StartARSession(UARSessionConfig* SessionConfig);
+	static UE_API void StartARSession(UARSessionConfig* SessionConfig);
 
 	/** Pause a running Augmented Reality session without clearing existing state. */
 	UFUNCTION(BlueprintCallable, Category = "AR AugmentedReality|Session", meta = (DisplayName="Pause AR Session", Keywords = "ar augmentedreality augmented reality session stop run running"))
-	static void PauseARSession();
+	static UE_API void PauseARSession();
 
 	/** Stop a running Augmented Reality session and clear any state. */
 	UFUNCTION(BlueprintCallable, Category = "AR AugmentedReality|Session", meta = (DisplayName="Stop AR Session", Keywords = "ar augmentedreality augmented reality session stop run running"))
-	static void StopARSession();
+	static UE_API void StopARSession();
 	
 	/**
 	 * It is intended that you check the status of the Augmented Reality session on every frame and take action accordingly.
@@ -83,23 +85,23 @@ public:
 	 * @return The status of a current Augmented Reality session: e.g. Running or Not running for a specific reason.
 	 */
 	UFUNCTION(BlueprintPure, Category = "AR AugmentedReality|Session", meta = (DisplayName="Get AR Session Status", Keywords = "ar augmentedreality augmented reality session start stop run running"))
-	static FARSessionStatus GetARSessionStatus();
+	static UE_API FARSessionStatus GetARSessionStatus();
 
 	/** @return the configuration that the current session was started with. */
 	UFUNCTION(BlueprintPure, Category = "AR AugmentedReality|Session", meta = (DisplayName="Get AR Session Config", Keywords = "ar augmentedreality augmented reality session"))
-	static UARSessionConfig* GetSessionConfig();
+	static UE_API UARSessionConfig* GetSessionConfig();
 
 	/** Starts or stops a battery intensive service on device. */
 	UFUNCTION(BlueprintCallable, Category = "AR AugmentedReality|Session", meta = (Keywords = "ar augmentedreality augmented reality capture start stop"))
-	static bool ToggleARCapture(const bool bOnOff, const EARCaptureType CaptureType);
+	static UE_API bool ToggleARCapture(const bool bOnOff, const EARCaptureType CaptureType);
 
 	/** Enable or disable Mixed Reality Capture camera. */
 	UFUNCTION(BlueprintCallable, Category = "AR", meta = (Keywords = "ar all"))
-	static void SetEnabledXRCamera(bool bOnOff);
+	static UE_API void SetEnabledXRCamera(bool bOnOff);
 
 	/** Change screen size of Mixed Reality Capture camera. */
 	UFUNCTION(BlueprintCallable, Category = "AR", meta = (Keywords = "ar all"))
-	static FIntPoint ResizeXRCamera(const FIntPoint& InSize);
+	static UE_API FIntPoint ResizeXRCamera(const FIntPoint& InSize);
 
 	
 	/**
@@ -115,7 +117,7 @@ public:
 	 * \see PinComponentToTraceResult
 	 */
 	UFUNCTION(BlueprintCallable, Category = "AR AugmentedReality|Alignment", meta = (DisplayName="Set AR Alignment Transform", Keywords = "ar augmentedreality augmented reality tracking alignment"))
-	static void SetAlignmentTransform( const FTransform& InAlignmentTransform );
+	static UE_API void SetAlignmentTransform( const FTransform& InAlignmentTransform );
 	
 	
 	/**
@@ -126,7 +128,7 @@ public:
 	 * @return a list of \c FARTraceResult sorted by distance from camera.
 	 */
 	UFUNCTION(BlueprintCallable, Category = "AR AugmentedReality|Trace Result", meta = (AdvancedDisplay="1", Keywords = "ar augmentedreality augmented reality tracking tracing linetrace"))
-	static TArray<FARTraceResult> LineTraceTrackedObjects( const FVector2D ScreenCoord, bool bTestFeaturePoints = true, bool bTestGroundPlane = true, bool bTestPlaneExtents = true, bool bTestPlaneBoundaryPolygon = true );
+	static UE_API TArray<FARTraceResult> LineTraceTrackedObjects( const FVector2D ScreenCoord, bool bTestFeaturePoints = true, bool bTestGroundPlane = true, bool bTestPlaneExtents = true, bool bTestPlaneBoundaryPolygon = true );
 	
 	/**
 	 * Perform a line trace against any real-world geometry as tracked by the AR system.
@@ -137,43 +139,43 @@ public:
 	 * @return a list of \c FARTraceResult sorted by distance from camera.
 	 */
 	UFUNCTION(BlueprintCallable, Category = "AR AugmentedReality|Trace Result", meta = (AdvancedDisplay = "1", Keywords = "ar augmentedreality augmented reality tracking tracing linetrace"))
-	static TArray<FARTraceResult> LineTraceTrackedObjects3D(const FVector Start, const FVector End, bool bTestFeaturePoints = true, bool bTestGroundPlane = true, bool bTestPlaneExtents = true, bool bTestPlaneBoundaryPolygon = true);
+	static UE_API TArray<FARTraceResult> LineTraceTrackedObjects3D(const FVector Start, const FVector End, bool bTestFeaturePoints = true, bool bTestGroundPlane = true, bool bTestPlaneExtents = true, bool bTestPlaneBoundaryPolygon = true);
 	
 	/** @return how well the tracking system is performing at the moment */
 	UFUNCTION(BlueprintPure, Category = "AR AugmentedReality|Tracking", meta = (DisplayName="Get AR Tracking Quality", Keywords = "ar augmentedreality augmented reality tracking quality"))
-	static EARTrackingQuality GetTrackingQuality();
+	static UE_API EARTrackingQuality GetTrackingQuality();
 	
 	/** @return The reason for the current limited tracking state */
 	UFUNCTION(BlueprintPure, Category = "AR AugmentedReality|Tracking", meta = (DisplayName="Get AR Tracking Quality Reason", Keywords = "ar augmentedreality augmented reality tracking quality reason"))
-	static EARTrackingQualityReason GetTrackingQualityReason();
+	static UE_API EARTrackingQualityReason GetTrackingQualityReason();
 	
 	/** @return a list of all the real-world geometry as currently seen by the Augmented Reality system */
 	UFUNCTION(BlueprintCallable, Category = "AR AugmentedReality|Tracking", meta = (DisplayName="Get All AR Geometries", Keywords = "ar augmentedreality augmented reality tracking geometry anchor"))
-	static TArray<UARTrackedGeometry*> GetAllGeometries();
+	static UE_API TArray<UARTrackedGeometry*> GetAllGeometries();
 	
 	/** @return a list of all the real-world geometry of the specified class as currently seen by the Augmented Reality system */
 	UFUNCTION(BlueprintCallable, Category = "AR AugmentedReality|Tracking", meta = (DisplayName="Get All AR Geometries By Class", Keywords = "ar augmentedreality augmented reality tracking geometry anchor", DeterminesOutputType = "GeometryClass"))
-	static TArray<UARTrackedGeometry*> GetAllGeometriesByClass(TSubclassOf<UARTrackedGeometry> GeometryClass);
+	static UE_API TArray<UARTrackedGeometry*> GetAllGeometriesByClass(TSubclassOf<UARTrackedGeometry> GeometryClass);
 	
 	/** @return the AR texture for the specified type */
 	UFUNCTION(BlueprintCallable, Category = "AR AugmentedReality|Texture", meta = (DisplayName="Get AR Texture", Keywords = "ar augmentedreality augmented reality texture"))
-	static UARTexture* GetARTexture(EARTextureType TextureType);
+	static UE_API UARTexture* GetARTexture(EARTextureType TextureType);
 	
 	/**
 	 * Test whether this type of session is supported by the current Augmented Reality platform.
 	 * e.g. is your device capable of doing positional tracking or orientation only?
 	 */
 	UFUNCTION(BlueprintPure, Category = "AR AugmentedReality|Tracking", meta = (DisplayName="Is AR Session Type Supported", Keywords = "ar augmentedreality augmented reality tracking"))
-	static bool IsSessionTypeSupported(EARSessionType SessionType);
+	static UE_API bool IsSessionTypeSupported(EARSessionType SessionType);
 	
 	
 	/** Given some real-world geometry being tracked by the Augmented Reality system, draw it on the screen for debugging purposes (rudimentary)*/
 	UFUNCTION(BlueprintCallable, Category = "AR AugmentedReality|Debug", meta = (WorldContext = "WorldContextObject", AdvancedDisplay="1", Keywords = "ar augmentedreality augmented reality tracked geometry debug draw"))
-	static void DebugDrawTrackedGeometry( UARTrackedGeometry* TrackedGeometry, UObject* WorldContextObject, FLinearColor Color = FLinearColor(1.0f, 1.0f, 0.0f, 0.75f), float OutlineThickness=5.0f, float PersistForSeconds = 0.0f );
+	static UE_API void DebugDrawTrackedGeometry( UARTrackedGeometry* TrackedGeometry, UObject* WorldContextObject, FLinearColor Color = FLinearColor(1.0f, 1.0f, 0.0f, 0.75f), float OutlineThickness=5.0f, float PersistForSeconds = 0.0f );
 	
 	/** Given a \c UARPin, draw it for debugging purposes. */
 	UFUNCTION(BlueprintCallable, Category = "AR AugmentedReality|Debug", meta = (WorldContext = "WorldContextObject", AdvancedDisplay="1", Keywords = "ar augmentedreality augmented reality pin debug draw"))
-	static void DebugDrawPin( UARPin* ARPin, UObject* WorldContextObject, FLinearColor Color = FLinearColor(1.0f, 1.0f, 0.0f, 0.75f), float Scale=5.0f, float PersistForSeconds = 0.0f );
+	static UE_API void DebugDrawPin( UARPin* ARPin, UObject* WorldContextObject, FLinearColor Color = FLinearColor(1.0f, 1.0f, 0.0f, 0.75f), float Scale=5.0f, float PersistForSeconds = 0.0f );
 	
 	
 	/**
@@ -186,7 +188,7 @@ public:
 	 * @return a \c UARLighEstimate that can be cast to a derived class.
 	 */
 	UFUNCTION( BlueprintPure, Category = "AR AugmentedReality|Light Estimate" )
-	static UARLightEstimate* GetCurrentLightEstimate();
+	static UE_API UARLightEstimate* GetCurrentLightEstimate();
 	
 	
 	/**
@@ -205,14 +207,14 @@ public:
 	 *         location and optionally to the \c TrackedGeometry.
 	 */
 	UFUNCTION(BlueprintCallable, Category = "AR AugmentedReality|ARPin", meta = (AdvancedDisplay="3", Keywords = "ar augmentedreality augmented reality tracking arpin tracked geometry pinning anchor"))
-	static UARPin* PinComponent( USceneComponent* ComponentToPin, const FTransform& PinToWorldTransform, UARTrackedGeometry* TrackedGeometry = nullptr, const FName DebugName = NAME_None );
+	static UE_API UARPin* PinComponent( USceneComponent* ComponentToPin, const FTransform& PinToWorldTransform, UARTrackedGeometry* TrackedGeometry = nullptr, const FName DebugName = NAME_None );
 	
 	/**
 	 * A convenient version of \c PinComponent() that can be used in conjunction
 	 * with a result of a \c LineTraceTrackedObjects call.
 	 */
 	UFUNCTION(BlueprintCallable, Category = "AR AugmentedReality|ARPin", meta = (AdvancedDisplay="2", Keywords = "ar augmentedreality augmented reality tracking arpin tracked geometry pinning anchor"))
-	static UARPin* PinComponentToTraceResult( USceneComponent* ComponentToPin, const FARTraceResult& TraceResult, const FName DebugName = NAME_None );
+	static UE_API UARPin* PinComponentToTraceResult( USceneComponent* ComponentToPin, const FARTraceResult& TraceResult, const FName DebugName = NAME_None );
 	
 	/**
 	 * Associate a component with an ARPin, so that its transform will be updated by the pin.  Any previously associated component will be detached.
@@ -223,19 +225,19 @@ public:
 	 * @return					True if the operation was successful.
 	 */
 	UFUNCTION(BlueprintCallable, Category = "AR AugmentedReality|ARPin", meta = (Keywords = "ar augmentedreality augmented reality tracking arpin tracked geometry pinning anchor"))
-	static bool PinComponentToARPin(USceneComponent* ComponentToPin, UARPin* Pin);
+	static UE_API bool PinComponentToARPin(USceneComponent* ComponentToPin, UARPin* Pin);
 
 	/** Given a pinned \c ComponentToUnpin, remove its attachment to the real world. */
 	UFUNCTION(BlueprintCallable, Category = "AR AugmentedReality|ARPin", meta = (Keywords = "ar augmentedreality augmented reality tracking arpin tracked geometry pinning anchor"))
-	static void UnpinComponent( USceneComponent* ComponentToUnpin );
+	static UE_API void UnpinComponent( USceneComponent* ComponentToUnpin );
 	
 	/** Remove a pin such that it no longer updates the associated component. */
 	UFUNCTION(BlueprintCallable, Category = "AR AugmentedReality|ARPin", meta = (Keywords = "ar augmentedreality augmented reality tracking arpin tracked geometry pinning anchor"))
-	static void RemovePin( UARPin* PinToRemove );
+	static UE_API void RemovePin( UARPin* PinToRemove );
 	
 	/** Get a list of all the \c UARPin objects that the Augmented Reality session is currently using to connect virtual objects to real-world, tracked locations. */
 	UFUNCTION(BlueprintCallable, Category = "AR AugmentedReality|ARPin", meta = (Keywords = "ar augmentedreality augmented reality tracking arpin anchor"))
-	static TArray<UARPin*> GetAllPins();
+	static UE_API TArray<UARPin*> GetAllPins();
 
 	/**
 	 * Is ARPin Local Store Supported
@@ -243,7 +245,7 @@ public:
 	 * @return					True if Local Pin saving is supported by the device/platform.
 	 */
 	UFUNCTION(BlueprintPure, Category = "AR AugmentedReality|ARPin", meta = (Keywords = "ar augmentedreality augmented reality tracking arpin anchor LocalStore"))
-	static bool IsARPinLocalStoreSupported();
+	static UE_API bool IsARPinLocalStoreSupported();
 
 	/**
 	 * Is ARPin Local Store Ready
@@ -251,7 +253,7 @@ public:
 	 * @return					True if local store is ready for use.
 	 */
 	UFUNCTION(BlueprintPure, Category = "AR AugmentedReality|ARPin", meta = (Keywords = "ar augmentedreality augmented reality tracking arpin anchor LocalStore"))
-	static bool IsARPinLocalStoreReady();
+	static UE_API bool IsARPinLocalStoreReady();
 
 	/**
 	 * Load all ARPins from local save
@@ -260,7 +262,7 @@ public:
 	 * @return					Map of SaveName:ARPin.
 	 */
 	UFUNCTION(BlueprintCallable, Category = "AR AugmentedReality|ARPin", meta = (Keywords = "ar augmentedreality augmented reality tracking arpin anchor LocalStore"))
-	static TMap<FName, UARPin*> LoadARPinsFromLocalStore();
+	static UE_API TMap<FName, UARPin*> LoadARPinsFromLocalStore();
 	
 	/**
 	 * Save an ARPin to local store
@@ -270,39 +272,39 @@ public:
 	 * @return					True if saved successfully.
 	 */
 	UFUNCTION(BlueprintCallable, Category = "AR AugmentedReality|ARPin", meta = (Keywords = "ar augmentedreality augmented reality tracking arpin anchor LocalStore"))
-	static bool SaveARPinToLocalStore(FName InSaveName, UARPin* InPin);
+	static UE_API bool SaveARPinToLocalStore(FName InSaveName, UARPin* InPin);
 	
 	/**
 	 * Remove an ARPin from the local store
 	 * @param InName			The save name to remove.
 	 */
 	UFUNCTION(BlueprintCallable, Category = "AR AugmentedReality|ARPin", meta = (Keywords = "ar augmentedreality augmented reality tracking arpin anchor LocalStore"))
-	static void RemoveARPinFromLocalStore(FName InSaveName);
+	static UE_API void RemoveARPinFromLocalStore(FName InSaveName);
 
 	/**
 	 * Remove all ARPins from the local store
 	 */
 	UFUNCTION(BlueprintCallable, Category = "AR AugmentedReality|ARPin", meta = (Keywords = "ar augmentedreality augmented reality tracking arpin anchor LocalStore"))
-	static void RemoveAllARPinsFromLocalStore();
+	static UE_API void RemoveAllARPinsFromLocalStore();
 
 	/** Adds an environment capture probe to the ar world */
 	UFUNCTION(BlueprintCallable, Category = "AR AugmentedReality|Tracking", meta = (DisplayName="Add AR Environment Probe", Keywords = "ar augmentedreality augmented reality tracking anchor"))
-	static bool AddManualEnvironmentCaptureProbe(FVector Location, FVector Extent);
+	static UE_API bool AddManualEnvironmentCaptureProbe(FVector Location, FVector Extent);
 
 	/** @return the current world mapping status for the AR world */
 	UFUNCTION(BlueprintPure, Category = "AR AugmentedReality|Tracking", meta = (DisplayName="Get AR World Mapping Status", Keywords = "ar augmentedreality augmented reality tracking anchor"))
-	static EARWorldMappingState GetWorldMappingStatus();
+	static UE_API EARWorldMappingState GetWorldMappingStatus();
 	
 	/** @return the raw point cloud view of the AR scene */
 	UFUNCTION(BlueprintPure, Category = "AR AugmentedReality|Tracking", meta = (DisplayName="Get AR Point Cloud", Keywords = "ar augmentedreality augmented reality tracking point cloud"))
-	static TArray<FVector> GetPointCloud();
+	static UE_API TArray<FVector> GetPointCloud();
 
 	/** @return The list of supported video formats for this device */
 	UFUNCTION(BlueprintCallable, Category = "AR AugmentedReality|Capabilities", meta = (DisplayName="Get Supported AR Video Formats", Keywords = "ar augmentedreality augmented reality config video formats"))
-	static TArray<FARVideoFormat> GetSupportedVideoFormats(EARSessionType SessionType);
+	static UE_API TArray<FARVideoFormat> GetSupportedVideoFormats(EARSessionType SessionType);
 
-	static TSharedPtr<FARSaveWorldAsyncTask, ESPMode::ThreadSafe> SaveWorld();
-	static TSharedPtr<FARGetCandidateObjectAsyncTask, ESPMode::ThreadSafe> GetCandidateObject(FVector Location, FVector Extent);
+	static UE_API TSharedPtr<FARSaveWorldAsyncTask, ESPMode::ThreadSafe> SaveWorld();
+	static UE_API TSharedPtr<FARGetCandidateObjectAsyncTask, ESPMode::ThreadSafe> GetCandidateObject(FVector Location, FVector Extent);
 	
 	/**
 	 * Create an ARCandidateImage object and add it to the ARCandidateImageList of the given \c UARSessionConfig object.
@@ -317,19 +319,19 @@ public:
 	 *		  Return nullptr otherwise.
 	 */
 	UFUNCTION(BlueprintCallable, Category = "AR AugmentedReality|Session", meta = (Keywords = "ar augmentedreality augmented reality candidate image"))
-	static UARCandidateImage* AddRuntimeCandidateImage(UARSessionConfig* SessionConfig, UTexture2D* CandidateTexture, FString FriendlyName, float PhysicalWidth);
+	static UE_API UARCandidateImage* AddRuntimeCandidateImage(UARSessionConfig* SessionConfig, UTexture2D* CandidateTexture, FString FriendlyName, float PhysicalWidth);
 
 	/** @return if a particular session feature is supported with the specified session type on the current platform */
 	UFUNCTION(BlueprintPure, Category = "AR AugmentedReality|Session", meta = (DisplayName="Is AR Session Tracking Feature Supported", Keywords = "ar augmentedreality augmented reality session tracking feature"))
-	static bool IsSessionTrackingFeatureSupported(EARSessionType SessionType, EARSessionTrackingFeature SessionTrackingFeature);
+	static UE_API bool IsSessionTrackingFeatureSupported(EARSessionType SessionType, EARSessionTrackingFeature SessionTrackingFeature);
 	
 	/** @return if a particular scene reconstruction method is supported with the specified session type on the current platform */
 	UFUNCTION(BlueprintPure, Category = "AR AugmentedReality|Session", meta = (DisplayName="Is AR Scene Reconstruction Supported", Keywords = "ar augmentedreality augmented reality session scene reconstruction"))
-	static bool IsSceneReconstructionSupported(EARSessionType SessionType, EARSceneReconstruction SceneReconstructionMethod);
+	static UE_API bool IsSceneReconstructionSupported(EARSessionType SessionType, EARSceneReconstruction SceneReconstructionMethod);
 	
 	/** @return all the 2D poses tracked by the AR system */
 	UFUNCTION(BlueprintCallable, Category = "AR AugmentedReality|Pose Tracking", meta = (DisplayName="Get All AR Tracked 2D Poses", Keywords = "ar augmentedreality augmented reality pose tracking"))
-	static TArray<FARPose2D> GetAllTracked2DPoses();
+	static UE_API TArray<FARPose2D> GetAllTracked2DPoses();
 	
 	/**
 	 * Try to determine the classification of the object at a world space location
@@ -340,7 +342,7 @@ public:
 	 * @return: whether a valid classification result is calculated
 	*/
 	UFUNCTION(BlueprintCallable, Category = "AR AugmentedReality|Classification")
-	static bool GetObjectClassificationAtLocation(const FVector& InWorldLocation, EARObjectClassification& OutClassification, FVector& OutClassificationLocation, float MaxLocationDiff = 10.f);
+	static UE_API bool GetObjectClassificationAtLocation(const FVector& InWorldLocation, EARObjectClassification& OutClassification, FVector& OutClassificationLocation, float MaxLocationDiff = 10.f);
 	
 	/**
 	 * For a point P in the AR local space, whose location and rotation are "OriginLocation" and "OriginRotation" in the world space
@@ -349,7 +351,7 @@ public:
 	 * @bMaintainUpDirection: if set, only the yaw roation of the alignment transform will be modified, pitch and roll will be zeroed out.
 	 */
 	UFUNCTION(BlueprintCallable, Category = "AR AugmentedReality|Alignment", meta = (DisplayName="Set AR World Origin Location and Rotation", Keywords = "ar augmentedreality augmented reality world origin"))
-	static void SetARWorldOriginLocationAndRotation(FVector OriginLocation, FRotator OriginRotation, bool bIsTransformInWorldSpace = true, bool bMaintainUpDirection = true);
+	static UE_API void SetARWorldOriginLocationAndRotation(FVector OriginLocation, FRotator OriginRotation, bool bIsTransformInWorldSpace = true, bool bMaintainUpDirection = true);
 	
 	/**
 	 * Helper function that modifies the alignment transform scale so that virtual content in the world space appears to be "scaled".
@@ -357,15 +359,15 @@ public:
 	 * moving the camera further away from the origin makes objects appear to be smaller, and vice versa.
 	 */
 	UFUNCTION(BlueprintCallable, Category = "AR AugmentedReality|Alignment", meta = (DisplayName="Set AR World Scale", Keywords = "ar augmentedreality augmented reality world scale"))
-	static void SetARWorldScale(float InWorldScale);
+	static UE_API void SetARWorldScale(float InWorldScale);
 	
 	/** @return the AR world scale, see "SetARWorldScale" */
 	UFUNCTION(BlueprintPure, Category = "AR AugmentedReality|Alignment", meta = (DisplayName="Get AR World Scale", Keywords = "ar augmentedreality augmented reality world scale"))
-	static float GetARWorldScale();
+	static UE_API float GetARWorldScale();
 	
 	/** @return the alignment transform, see "SetAlignmentTransform" */
 	UFUNCTION(BlueprintPure, Category = "AR AugmentedReality|Alignment", meta = (DisplayName="Get Alignment Transform", Keywords = "ar augmentedreality augmented reality alignment transform"))
-	static FTransform GetAlignmentTransform();
+	static UE_API FTransform GetAlignmentTransform();
 	
 	/**
 	 * Manually add a tracked point with name and world transform.
@@ -376,11 +378,11 @@ public:
 	 * Note that this is an async operation - the added point won't be available until a few frames later.
 	 */
 	UFUNCTION(BlueprintCallable, Category = "AR AugmentedReality", meta = (Keywords = "ar augmentedreality augmented reality add point"))
-	static bool AddTrackedPointWithName(const FTransform& WorldTransform, const FString& PointName, bool bDeletePointsWithSameName = true);
+	static UE_API bool AddTrackedPointWithName(const FTransform& WorldTransform, const FString& PointName, bool bDeletePointsWithSameName = true);
 	
 	/** @return a list of the tracked points with the given name */
 	UFUNCTION(BlueprintCallable, Category = "AR AugmentedReality", meta = (Keywords = "ar augmentedreality augmented reality find point"))
-	static TArray<UARTrackedPoint*> FindTrackedPointsByName(const FString& PointName);
+	static UE_API TArray<UARTrackedPoint*> FindTrackedPointsByName(const FString& PointName);
 	
 	// Static helpers to create the methods needed to add/remove delegates from the AR system
 	DEFINE_AR_BPLIB_DELEGATE_FUNCS(OnTrackableAdded)
@@ -390,19 +392,19 @@ public:
 
 	//Alignment helpers
 	UFUNCTION(BlueprintCallable, Category = Alignment)
-	static void CalculateClosestIntersection(const TArray<FVector>& StartPoints, const TArray<FVector>& EndPoints, FVector& ClosestIntersection);
+	static UE_API void CalculateClosestIntersection(const TArray<FVector>& StartPoints, const TArray<FVector>& EndPoints, FVector& ClosestIntersection);
 
 	// Computes a transform that aligns two coordinate systems. Requires the transform of the same known point in each coordinate system.
 	UFUNCTION(BlueprintCallable, Category = Alignment)
-	static void CalculateAlignmentTransform(const FTransform& TransformInFirstCoordinateSystem, const FTransform& TransformInSecondCoordinateSystem, FTransform& AlignmentTransform);
+	static UE_API void CalculateAlignmentTransform(const FTransform& TransformInFirstCoordinateSystem, const FTransform& TransformInSecondCoordinateSystem, FTransform& AlignmentTransform);
 	
 	/** @return the max number of faces can be tracked at the same time */
 	UFUNCTION(BlueprintCallable, Category = "AR AugmentedReality|Face Tracking")
-	static int32 GetNumberOfTrackedFacesSupported();
+	static UE_API int32 GetNumberOfTrackedFacesSupported();
 	
 	/** @return the intrinsics of the AR camera. */
 	UFUNCTION(BlueprintCallable, Category = "AR AugmentedReality|Geo Tracking")
-	static bool GetCameraIntrinsics(FARCameraIntrinsics& OutCameraIntrinsics);
+	static UE_API bool GetCameraIntrinsics(FARCameraIntrinsics& OutCameraIntrinsics);
 	
 	/** @return a list of all the real-world geometry of the specified class as currently seen by the Augmented Reality system */
 	template<class T>
@@ -420,16 +422,16 @@ public:
 	}
 	
 public:
-	static void RegisterAsARSystem(const TSharedRef<FARSupportInterface , ESPMode::ThreadSafe>& NewArSystem);
+	static UE_API void RegisterAsARSystem(const TSharedRef<FARSupportInterface , ESPMode::ThreadSafe>& NewArSystem);
 	
 private:
-	static const TWeakPtr<FARSupportInterface , ESPMode::ThreadSafe>& GetARSystem();
-	static TWeakPtr<FARSupportInterface , ESPMode::ThreadSafe> RegisteredARSystem;
+	static UE_API const TWeakPtr<FARSupportInterface , ESPMode::ThreadSafe>& GetARSystem();
+	static UE_API TWeakPtr<FARSupportInterface , ESPMode::ThreadSafe> RegisteredARSystem;
 };
 
 
-UCLASS(meta=(ScriptName="ARTraceResultLibrary"))
-class AUGMENTEDREALITY_API UARTraceResultLibrary : public UBlueprintFunctionLibrary
+UCLASS(MinimalAPI, meta=(ScriptName="ARTraceResultLibrary"))
+class UARTraceResultLibrary : public UBlueprintFunctionLibrary
 {
 	GENERATED_BODY()
 	
@@ -437,7 +439,7 @@ class AUGMENTEDREALITY_API UARTraceResultLibrary : public UBlueprintFunctionLibr
 	
 	/** @return  the distance from the camera to the traced location in Unreal Units. */
 	UFUNCTION( BlueprintPure, Category = "AR AugmentedReality|Trace Result" )
-	static float GetDistanceFromCamera( const FARTraceResult& TraceResult );
+	static UE_API float GetDistanceFromCamera( const FARTraceResult& TraceResult );
 	
 	/**
 	 * @return The transform of the trace result in tracking space (after it is modified by the \c AlignmentTransform).
@@ -445,21 +447,23 @@ class AUGMENTEDREALITY_API UARTraceResultLibrary : public UBlueprintFunctionLibr
 	 * \see SetAlignmentTransform()
 	 */
 	UFUNCTION( BlueprintPure, Category = "AR AugmentedReality|Trace Result" )
-	static FTransform GetLocalToTrackingTransform( const FARTraceResult& TraceResult );
+	static UE_API FTransform GetLocalToTrackingTransform( const FARTraceResult& TraceResult );
 	
 	/** @return Get the transform of the trace result in Unreal World Space. */
 	UFUNCTION( BlueprintPure, Category = "AR AugmentedReality|Trace Result" )
-	static FTransform GetLocalToWorldTransform( const FARTraceResult& TraceResult );
+	static UE_API FTransform GetLocalToWorldTransform( const FARTraceResult& TraceResult );
 	
 	/** @return Get the transform of the trace result in the AR system's local space. */
 	UFUNCTION( BlueprintPure, Category = "AR AugmentedReality|Trace Result" )
-	static FTransform GetLocalTransform( const FARTraceResult& TraceResult );
+	static UE_API FTransform GetLocalTransform( const FARTraceResult& TraceResult );
 	
 	/** @return Get the real-world object (as observed by the Augmented Reality system) that was intersected by the line trace. */
 	UFUNCTION( BlueprintPure, Category = "AR AugmentedReality|Trace Result" )
-	static UARTrackedGeometry* GetTrackedGeometry( const FARTraceResult& TraceResult );
+	static UE_API UARTrackedGeometry* GetTrackedGeometry( const FARTraceResult& TraceResult );
 	
 	/** @return  Get the type of the tracked object (if any) that effected this trace result. */
 	UFUNCTION( BlueprintPure, Category = "AR AugmentedReality|Trace Result" )
-	static EARLineTraceChannels GetTraceChannel( const FARTraceResult& TraceResult );
+	static UE_API EARLineTraceChannels GetTraceChannel( const FARTraceResult& TraceResult );
 };
+
+#undef UE_API
diff --git a/Engine/Source/Runtime/AugmentedReality/Public/ARBlueprintProxy.h b/Engine/Source/Runtime/AugmentedReality/Public/ARBlueprintProxy.h
--- a/Engine/Source/Runtime/AugmentedReality/Public/ARBlueprintProxy.h
+++ b/Engine/Source/Runtime/AugmentedReality/Public/ARBlueprintProxy.h
@@ -14,8 +14,10 @@
 
 #include "ARBlueprintProxy.generated.h"
 
-UCLASS(Abstract)
-class AUGMENTEDREALITY_API UARBaseAsyncTaskBlueprintProxy :
+#define UE_API AUGMENTEDREALITY_API
+
+UCLASS(MinimalAPI, Abstract)
+class UARBaseAsyncTaskBlueprintProxy :
 	public UBlueprintAsyncActionBase,
 	public FTickableGameObject
 {
@@ -23,7 +25,7 @@ class AUGMENTEDREALITY_API UARBaseAsyncTaskBlueprintProxy :
 	
 public:
 	//~ Begin FTickableObject Interface
-	virtual void Tick(float DeltaTime) override;
+	UE_API virtual void Tick(float DeltaTime) override;
 	virtual bool IsTickable() const override { return !HasAnyFlags(RF_ClassDefaultObject) && bShouldTick; }
 	virtual TStatId GetStatId() const override { RETURN_QUICK_DECLARE_CYCLE_STAT(UARBaseAsyncTaskBlueprintProxy, STATGROUP_Tickables); }
 	//~ End FTickableObject Interface
@@ -31,10 +33,10 @@ public:
 	virtual void ReportSuccess() { check(0); }
 	virtual void ReportFailure() { check(0); }
 
-	static void RegisterAsARSystem(const TSharedRef<FARSupportInterface , ESPMode::ThreadSafe>& NewArSystem);
+	static UE_API void RegisterAsARSystem(const TSharedRef<FARSupportInterface , ESPMode::ThreadSafe>& NewArSystem);
 
 protected:
-	static const TWeakPtr<FARSupportInterface , ESPMode::ThreadSafe>& GetARSystem();
+	static UE_API const TWeakPtr<FARSupportInterface , ESPMode::ThreadSafe>& GetARSystem();
 	/** The async task to check during Tick() */
 	TSharedPtr<FARAsyncTask, ESPMode::ThreadSafe> AsyncTask;
 	
@@ -42,7 +44,7 @@ private:
 	/** True until the async task completes, then false */
 	bool bShouldTick;
 
-	static TWeakPtr<FARSupportInterface , ESPMode::ThreadSafe> RegisteredARSystem;
+	static UE_API TWeakPtr<FARSupportInterface , ESPMode::ThreadSafe> RegisteredARSystem;
 };
 
 DECLARE_DYNAMIC_MULTICAST_DELEGATE_OneParam(FARSaveWorldPin, const TArray<uint8>&, SavedWorld);
@@ -118,3 +120,5 @@ private:
 	virtual void ReportSuccess() override;
 	virtual void ReportFailure() override;
 };
+
+#undef UE_API
diff --git a/Engine/Source/Runtime/AugmentedReality/Public/ARComponent.h b/Engine/Source/Runtime/AugmentedReality/Public/ARComponent.h
--- a/Engine/Source/Runtime/AugmentedReality/Public/ARComponent.h
+++ b/Engine/Source/Runtime/AugmentedReality/Public/ARComponent.h
@@ -8,6 +8,8 @@
 #include "MRMeshBufferDefines.h"
 #include "ARComponent.generated.h"
 
+#define UE_API AUGMENTEDREALITY_API
+
 class USceneComponent;
 class UMRMeshComponent;
 class UARSessionConfig;
@@ -205,7 +207,7 @@ struct FARGeoAnchorUpdatePayload
 	FString AnchorName;
 };
 
-struct AUGMENTEDREALITY_API FAccumulatedNormal
+struct FAccumulatedNormal
 {
 	FVector Normal = FVector::ZeroVector;
 	uint32 NumFaces = 0;
@@ -227,31 +229,31 @@ struct AUGMENTEDREALITY_API FAccumulatedNormal
  * A MRMeshComponent will be created if mesh visualization is enabled (see UARSessionConfig::bGenerateMeshDataFromTrackedGeometry)
  * It can be accessed via 'GetMRMesh' or 'UARTrackedGeometry::GetUnderlyingMesh'
  */
-UCLASS(abstract, BlueprintType, Category = "AR Gameplay")
-class AUGMENTEDREALITY_API UARComponent : public USceneComponent
+UCLASS(MinimalAPI, abstract, BlueprintType, Category = "AR Gameplay")
+class UARComponent : public USceneComponent
 {
 	GENERATED_BODY()
 
-	virtual void BeginPlay() override;
+	UE_API virtual void BeginPlay() override;
 
 public:
-	UARComponent();
+	UE_API UARComponent();
 	
 	UFUNCTION(BlueprintCallable, Category = "AR Gameplay")
-	void SetNativeID(FGuid NativeID);
+	UE_API void SetNativeID(FGuid NativeID);
 	
 	UFUNCTION(BlueprintPure, Category = "AR Gameplay")
 	UMRMeshComponent* GetMRMesh() { return MRMeshComponent; }
 	
 	/** Event when native representation is removed, called on server and clients. */
 	UFUNCTION(BlueprintImplementableEvent, meta = (DisplayName = "Remove"))
-	void ReceiveRemove();
+	UE_API void ReceiveRemove();
 	
 	UFUNCTION(BlueprintCallable, BlueprintNativeEvent, meta = (DisplayName = "Update Visualization"), Category = "AR Gameplay")
-	void UpdateVisualization();
+	UE_API void UpdateVisualization();
 	
 	virtual void Update(UARTrackedGeometry* TrackedGeometry) {};
-	virtual void Remove(UARTrackedGeometry* TrackedGeometry);
+	UE_API virtual void Remove(UARTrackedGeometry* TrackedGeometry);
 
 	UPROPERTY(Replicated)
 	FGuid NativeID;
@@ -261,16 +263,16 @@ public:
 	FMRMeshDelegate OnMRMeshDestroyed;
 	
 protected:
-	virtual void OnUnregister() override;
+	UE_API virtual void OnUnregister() override;
 	
 	UFUNCTION()
-	virtual void OnRep_Payload();
+	UE_API virtual void OnRep_Payload();
 	
-	void RemoveMeshComponent(UARTrackedGeometry* TrackedGeometry);
+	UE_API void RemoveMeshComponent(UARTrackedGeometry* TrackedGeometry);
 	
-	void ManageMeshComponentForDebugMode(bool bDebugModeEnabled, const FARSessionPayload& SessionPayload);
+	UE_API void ManageMeshComponentForDebugMode(bool bDebugModeEnabled, const FARSessionPayload& SessionPayload);
 	
-	FLinearColor GetDebugColor() const;
+	UE_API FLinearColor GetDebugColor() const;
 	
 	/** If the default replication logic should be used for this component */
 	UPROPERTY(EditAnywhere, Category = "AR Gameplay")
@@ -321,71 +323,71 @@ enum class EPlaneComponentDebugMode : uint8
 	ShowClassification = 2,
 };
 
-UCLASS(Blueprintable, BlueprintType, Category = "AR Gameplay")
-class AUGMENTEDREALITY_API UARPlaneComponent : public UARComponent
+UCLASS(MinimalAPI, Blueprintable, BlueprintType, Category = "AR Gameplay")
+class UARPlaneComponent : public UARComponent
 {
 	GENERATED_BODY()
 
 public:
-	virtual void TickComponent(float DeltaTime, enum ELevelTick TickType, FActorComponentTickFunction* ThisTickFunction) override;
+	UE_API virtual void TickComponent(float DeltaTime, enum ELevelTick TickType, FActorComponentTickFunction* ThisTickFunction) override;
 	
 	/** Event when native representation is first added, called on server and clients. */
 	UFUNCTION(BlueprintImplementableEvent, meta = (DisplayName = "Add"))
-	void ReceiveAdd(const FARPlaneUpdatePayload& Payload);
+	UE_API void ReceiveAdd(const FARPlaneUpdatePayload& Payload);
 	
 	/** Event when native representation is updated, called on server and clients. */
 	UFUNCTION(BlueprintImplementableEvent, meta = (DisplayName = "Update"))
-	void ReceiveUpdate(const FARPlaneUpdatePayload& Payload);
+	UE_API void ReceiveUpdate(const FARPlaneUpdatePayload& Payload);
 	
 	/** Overridable native event for when native representation is updated. */
-	virtual void Update(UARTrackedGeometry* TrackedGeometry) override;
+	UE_API virtual void Update(UARTrackedGeometry* TrackedGeometry) override;
 	
 	UFUNCTION(BlueprintCallable, Category = "AR Gameplay | Debug", meta=(DevelopmentOnly))
-	static void SetPlaneComponentDebugMode(EPlaneComponentDebugMode NewDebugMode);
+	static UE_API void SetPlaneComponentDebugMode(EPlaneComponentDebugMode NewDebugMode);
 	
 	UFUNCTION(BlueprintCallable, Category = "AR Gameplay | Debug", meta=(DevelopmentOnly))
-	static void SetObjectClassificationDebugColors(const TMap<EARObjectClassification, FLinearColor>& InColors);
+	static UE_API void SetObjectClassificationDebugColors(const TMap<EARObjectClassification, FLinearColor>& InColors);
 	
 	UFUNCTION(BlueprintPure, Category = "AR Gameplay | Debug", meta=(DevelopmentOnly))
-	static const TMap<EARObjectClassification, FLinearColor>& GetObjectClassificationDebugColors();
+	static UE_API const TMap<EARObjectClassification, FLinearColor>& GetObjectClassificationDebugColors();
 	
 protected:
 	UFUNCTION(Reliable, Server, WithValidation)
-	void ServerUpdatePayload(const FARPlaneUpdatePayload& NewPayload);
+	UE_API void ServerUpdatePayload(const FARPlaneUpdatePayload& NewPayload);
 	
 	UPROPERTY(BlueprintReadOnly, ReplicatedUsing=OnRep_Payload, Category = "AR Gameplay")
 	FARPlaneUpdatePayload ReplicatedPayload;
 	
-	virtual void OnRep_Payload() override;
-	virtual void UpdateVisualization_Implementation() override;
-	static int32 GetDebugMode();
+	UE_API virtual void OnRep_Payload() override;
+	UE_API virtual void UpdateVisualization_Implementation() override;
+	static UE_API int32 GetDebugMode();
 };
 
 
-UCLASS(Blueprintable, BlueprintType, Category = "AR Gameplay")
-class AUGMENTEDREALITY_API UARPointComponent : public UARComponent
+UCLASS(MinimalAPI, Blueprintable, BlueprintType, Category = "AR Gameplay")
+class UARPointComponent : public UARComponent
 {
 	GENERATED_BODY()
 public:
 	/** Event when native representation is first added, called on server and clients. */
 	UFUNCTION(BlueprintImplementableEvent, meta = (DisplayName = "Add"))
-	void ReceiveAdd(const FARPointUpdatePayload& Payload);
+	UE_API void ReceiveAdd(const FARPointUpdatePayload& Payload);
 	
 	/** Event when native representation is updated, called on server and clients. */
 	UFUNCTION(BlueprintImplementableEvent, meta = (DisplayName = "Update"))
-	void ReceiveUpdate(const FARPointUpdatePayload& Payload);
+	UE_API void ReceiveUpdate(const FARPointUpdatePayload& Payload);
 	
 	/** Overridable native event for when native representation is updated. */
-	virtual void Update(UARTrackedGeometry* TrackedGeometry) override;
+	UE_API virtual void Update(UARTrackedGeometry* TrackedGeometry) override;
 	
 protected:
 	UFUNCTION(Reliable, Server, WithValidation)
-	void ServerUpdatePayload(const FARPointUpdatePayload& NewPayload);
+	UE_API void ServerUpdatePayload(const FARPointUpdatePayload& NewPayload);
 	
 	UPROPERTY(BlueprintReadOnly, ReplicatedUsing=OnRep_Payload, Category = "AR Gameplay")
 	FARPointUpdatePayload ReplicatedPayload;
 	
-	virtual void OnRep_Payload() override;
+	UE_API virtual void OnRep_Payload() override;
 };
 
 UENUM(BlueprintType)
@@ -417,27 +419,27 @@ enum class EARFaceTransformMixing : uint8
 	TrackingOnly
 };
 
-UCLASS(Blueprintable, BlueprintType, Category = "AR Gameplay")
-class AUGMENTEDREALITY_API UARFaceComponent : public UARComponent
+UCLASS(MinimalAPI, Blueprintable, BlueprintType, Category = "AR Gameplay")
+class UARFaceComponent : public UARComponent
 {
 	GENERATED_BODY()
 
 public:
-	virtual void TickComponent(float DeltaTime, enum ELevelTick TickType, FActorComponentTickFunction* ThisTickFunction) override;
+	UE_API virtual void TickComponent(float DeltaTime, enum ELevelTick TickType, FActorComponentTickFunction* ThisTickFunction) override;
 	
 	/** Event when native representation is first added, called on server and clients. */
 	UFUNCTION(BlueprintImplementableEvent, meta = (DisplayName = "Add"))
-	void ReceiveAdd(const FARFaceUpdatePayload& Payload);
+	UE_API void ReceiveAdd(const FARFaceUpdatePayload& Payload);
 	
 	/** Event when native representation is updated, called on server and clients. */
 	UFUNCTION(BlueprintImplementableEvent, meta = (DisplayName = "Update"))
-	void ReceiveUpdate(const FARFaceUpdatePayload& Payload);
+	UE_API void ReceiveUpdate(const FARFaceUpdatePayload& Payload);
 	
 	/** Overridable native event for when native representation is updated. */
-	virtual void Update(UARTrackedGeometry* TrackedGeometry) override;
+	UE_API virtual void Update(UARTrackedGeometry* TrackedGeometry) override;
 	
 	UFUNCTION(BlueprintCallable, Category = "AR Gameplay | Debug", meta=(DevelopmentOnly))
-	static void SetFaceComponentDebugMode(EFaceComponentDebugMode NewDebugMode);
+	static UE_API void SetFaceComponentDebugMode(EFaceComponentDebugMode NewDebugMode);
 	
 protected:
 	/**	Determines how the transform from tracking data and the component's transform are mixed together. */
@@ -453,14 +455,14 @@ protected:
 	bool bFaceOutOfScreen = true;
 	
 	UFUNCTION(Reliable, Server, WithValidation)
-	void ServerUpdatePayload(const FARFaceUpdatePayload& NewPayload);
+	UE_API void ServerUpdatePayload(const FARFaceUpdatePayload& NewPayload);
 	
 	UPROPERTY(BlueprintReadOnly, ReplicatedUsing=OnRep_Payload, Category = "AR Gameplay")
 	FARFaceUpdatePayload ReplicatedPayload;
 	
-	virtual void OnRep_Payload() override;
-	virtual void UpdateVisualization_Implementation() override;
-	static int32 GetDebugMode();
+	UE_API virtual void OnRep_Payload() override;
+	UE_API virtual void UpdateVisualization_Implementation() override;
+	static UE_API int32 GetDebugMode();
 	
 private:
 	/** Temporary data used for calculating the normals from vertices */
@@ -479,38 +481,38 @@ enum class EImageComponentDebugMode : uint8
 	ShowDetectedImage = 1,
 };
 
-UCLASS(Blueprintable, BlueprintType, Category = "AR Gameplay")
-class AUGMENTEDREALITY_API UARImageComponent : public UARComponent
+UCLASS(MinimalAPI, Blueprintable, BlueprintType, Category = "AR Gameplay")
+class UARImageComponent : public UARComponent
 {
 	GENERATED_BODY()
 	
 public:
-	virtual void TickComponent(float DeltaTime, enum ELevelTick TickType, FActorComponentTickFunction* ThisTickFunction) override;
+	UE_API virtual void TickComponent(float DeltaTime, enum ELevelTick TickType, FActorComponentTickFunction* ThisTickFunction) override;
 	
 	/** Event when native representation is first added, called on server and clients. */
 	UFUNCTION(BlueprintImplementableEvent, meta = (DisplayName = "Add"))
-	void ReceiveAdd(const FARImageUpdatePayload& Payload);
+	UE_API void ReceiveAdd(const FARImageUpdatePayload& Payload);
 	
 	/** Event when native representation is updated, called on server and clients. */
 	UFUNCTION(BlueprintImplementableEvent, meta = (DisplayName = "Update"))
-	void ReceiveUpdate(const FARImageUpdatePayload& Payload);
+	UE_API void ReceiveUpdate(const FARImageUpdatePayload& Payload);
 	
 	UFUNCTION(BlueprintCallable, Category = "AR Gameplay | Debug", meta=(DevelopmentOnly))
-	static void SetImageComponentDebugMode(EImageComponentDebugMode NewDebugMode);
+	static UE_API void SetImageComponentDebugMode(EImageComponentDebugMode NewDebugMode);
 	
 	/** Overridable native event for when native representation is updated. */
-	virtual void Update(UARTrackedGeometry* TrackedGeometry) override;
+	UE_API virtual void Update(UARTrackedGeometry* TrackedGeometry) override;
 	
 protected:
 	UFUNCTION(Reliable, Server, WithValidation)
-	void ServerUpdatePayload(const FARImageUpdatePayload& NewPayload);
+	UE_API void ServerUpdatePayload(const FARImageUpdatePayload& NewPayload);
 	
 	UPROPERTY(BlueprintReadOnly, ReplicatedUsing=OnRep_Payload, Category = "AR Gameplay")
 	FARImageUpdatePayload ReplicatedPayload;
 	
-	virtual void OnRep_Payload() override;
-	virtual void UpdateVisualization_Implementation() override;
-	static int32 GetDebugMode();
+	UE_API virtual void OnRep_Payload() override;
+	UE_API virtual void UpdateVisualization_Implementation() override;
+	static UE_API int32 GetDebugMode();
 };
 
 UENUM(BlueprintType)
@@ -523,38 +525,38 @@ enum class EQRCodeComponentDebugMode : uint8
 	ShowQRCode = 1,
 };
 
-UCLASS(Blueprintable, BlueprintType, Category = "AR Gameplay")
-class AUGMENTEDREALITY_API UARQRCodeComponent : public UARComponent
+UCLASS(MinimalAPI, Blueprintable, BlueprintType, Category = "AR Gameplay")
+class UARQRCodeComponent : public UARComponent
 {
 	GENERATED_BODY()
 
 public:
-	virtual void TickComponent(float DeltaTime, enum ELevelTick TickType, FActorComponentTickFunction* ThisTickFunction) override;
+	UE_API virtual void TickComponent(float DeltaTime, enum ELevelTick TickType, FActorComponentTickFunction* ThisTickFunction) override;
 	
 	/** Event when native representation is first added, called on server and clients. */
 	UFUNCTION(BlueprintImplementableEvent, meta = (DisplayName = "Add"))
-	void ReceiveAdd(const FARQRCodeUpdatePayload& Payload);
+	UE_API void ReceiveAdd(const FARQRCodeUpdatePayload& Payload);
 	
 	/** Event when native representation is updated, called on server and clients. */
 	UFUNCTION(BlueprintImplementableEvent, meta = (DisplayName = "Update"))
-	void ReceiveUpdate(const FARQRCodeUpdatePayload& Payload);
+	UE_API void ReceiveUpdate(const FARQRCodeUpdatePayload& Payload);
 	
 	/** Overridable native event for when native representation is updated. */
-	virtual void Update(UARTrackedGeometry* TrackedGeometry) override;
+	UE_API virtual void Update(UARTrackedGeometry* TrackedGeometry) override;
 	
 	UFUNCTION(BlueprintCallable, Category = "AR Gameplay | Debug", meta=(DevelopmentOnly))
-	static void SetQRCodeComponentDebugMode(EQRCodeComponentDebugMode NewDebugMode);
+	static UE_API void SetQRCodeComponentDebugMode(EQRCodeComponentDebugMode NewDebugMode);
 	
 protected:
 	UFUNCTION(Reliable, Server, WithValidation)
-	void ServerUpdatePayload(const FARQRCodeUpdatePayload& NewPayload);
+	UE_API void ServerUpdatePayload(const FARQRCodeUpdatePayload& NewPayload);
 	
 	UPROPERTY(BlueprintReadOnly, ReplicatedUsing=OnRep_Payload, Category = "AR Gameplay")
 	FARQRCodeUpdatePayload ReplicatedPayload;
 	
-	virtual void OnRep_Payload() override;
-	virtual void UpdateVisualization_Implementation() override;
-	static int32 GetDebugMode();
+	UE_API virtual void OnRep_Payload() override;
+	UE_API virtual void UpdateVisualization_Implementation() override;
+	static UE_API int32 GetDebugMode();
 };
 
 UENUM(BlueprintType)
@@ -567,119 +569,119 @@ enum class EPoseComponentDebugMode : uint8
 	ShowSkeleton = 1,
 };
 
-UCLASS(Blueprintable, BlueprintType, Category = "AR Gameplay")
-class AUGMENTEDREALITY_API UARPoseComponent : public UARComponent
+UCLASS(MinimalAPI, Blueprintable, BlueprintType, Category = "AR Gameplay")
+class UARPoseComponent : public UARComponent
 {
 	GENERATED_BODY()
 
 public:
-	virtual void TickComponent(float DeltaTime, enum ELevelTick TickType, FActorComponentTickFunction* ThisTickFunction) override;
+	UE_API virtual void TickComponent(float DeltaTime, enum ELevelTick TickType, FActorComponentTickFunction* ThisTickFunction) override;
 	
 	/** Event when native representation is first added, called on server and clients. */
 	UFUNCTION(BlueprintImplementableEvent, meta = (DisplayName = "Add"))
-	void ReceiveAdd(const FARPoseUpdatePayload& Payload);
+	UE_API void ReceiveAdd(const FARPoseUpdatePayload& Payload);
 	
 	/** Event when native representation is updated, called on server and clients. */
 	UFUNCTION(BlueprintImplementableEvent, meta = (DisplayName = "Update"))
-	void ReceiveUpdate(const FARPoseUpdatePayload& Payload);
+	UE_API void ReceiveUpdate(const FARPoseUpdatePayload& Payload);
 	
 	/** Overridable native event for when native representation is updated. */
-	virtual void Update(UARTrackedGeometry* TrackedGeometry) override;
+	UE_API virtual void Update(UARTrackedGeometry* TrackedGeometry) override;
 	
 	UFUNCTION(BlueprintCallable, Category = "AR Gameplay | Debug", meta=(DevelopmentOnly))
-	static void SetPoseComponentDebugMode(EPoseComponentDebugMode NewDebugMode);
+	static UE_API void SetPoseComponentDebugMode(EPoseComponentDebugMode NewDebugMode);
 	
 protected:
 	UFUNCTION(Reliable, Server, WithValidation)
-	void ServerUpdatePayload(const FARPoseUpdatePayload& NewPayload);
+	UE_API void ServerUpdatePayload(const FARPoseUpdatePayload& NewPayload);
 	
 	UPROPERTY(BlueprintReadOnly, ReplicatedUsing=OnRep_Payload, Category = "AR Gameplay")
 	FARPoseUpdatePayload ReplicatedPayload;
 	
-	virtual void OnRep_Payload() override;
-	static int32 GetDebugMode();
+	UE_API virtual void OnRep_Payload() override;
+	static UE_API int32 GetDebugMode();
 };
 
 
-UCLASS(Blueprintable, BlueprintType, Category = "AR Gameplay")
-class AUGMENTEDREALITY_API UAREnvironmentProbeComponent : public UARComponent
+UCLASS(MinimalAPI, Blueprintable, BlueprintType, Category = "AR Gameplay")
+class UAREnvironmentProbeComponent : public UARComponent
 {
 	GENERATED_BODY()
 public:
 	/** Event when native representation is first added, called on server and clients. */
 	UFUNCTION(BlueprintImplementableEvent, meta = (DisplayName = "Add"))
-	void ReceiveAdd(const FAREnvironmentProbeUpdatePayload& Payload);
+	UE_API void ReceiveAdd(const FAREnvironmentProbeUpdatePayload& Payload);
 	
 	/** Event when native representation is updated, called on server and clients. */
 	UFUNCTION(BlueprintImplementableEvent, meta = (DisplayName = "Update"))
-	void ReceiveUpdate(const FAREnvironmentProbeUpdatePayload& Payload);
+	UE_API void ReceiveUpdate(const FAREnvironmentProbeUpdatePayload& Payload);
 	
 	/** Overridable native event for when native representation is updated. */
-	virtual void Update(UARTrackedGeometry* TrackedGeometry) override;
+	UE_API virtual void Update(UARTrackedGeometry* TrackedGeometry) override;
 	
 protected:
 	UFUNCTION(Reliable, Server, WithValidation)
-	void ServerUpdatePayload(const FAREnvironmentProbeUpdatePayload& NewPayload);
+	UE_API void ServerUpdatePayload(const FAREnvironmentProbeUpdatePayload& NewPayload);
 	
 	UPROPERTY(BlueprintReadOnly, ReplicatedUsing=OnRep_Payload, Category = "AR Gameplay")
 	FAREnvironmentProbeUpdatePayload ReplicatedPayload;
 	
-	virtual void OnRep_Payload() override;
+	UE_API virtual void OnRep_Payload() override;
 };
 
 
-UCLASS(Blueprintable, BlueprintType, Category = "AR Gameplay")
-class AUGMENTEDREALITY_API UARObjectComponent : public UARComponent
+UCLASS(MinimalAPI, Blueprintable, BlueprintType, Category = "AR Gameplay")
+class UARObjectComponent : public UARComponent
 {
 	GENERATED_BODY()
 public:
 	/** Event when native representation is first added, called on server and clients. */
 	UFUNCTION(BlueprintImplementableEvent, meta = (DisplayName = "Add"))
-	void ReceiveAdd(const FARObjectUpdatePayload& Payload);
+	UE_API void ReceiveAdd(const FARObjectUpdatePayload& Payload);
 	
 	/** Event when native representation is updated, called on server and clients. */
 	UFUNCTION(BlueprintImplementableEvent, meta = (DisplayName = "Update"))
-	void ReceiveUpdate(const FARObjectUpdatePayload& Payload);
+	UE_API void ReceiveUpdate(const FARObjectUpdatePayload& Payload);
 	
 	/** Overridable native event for when native representation is updated. */
-	virtual void Update(UARTrackedGeometry* TrackedGeometry) override;
+	UE_API virtual void Update(UARTrackedGeometry* TrackedGeometry) override;
 	
 protected:
 	UFUNCTION(Reliable, Server, WithValidation)
-	void ServerUpdatePayload(const FARObjectUpdatePayload& NewPayload);
+	UE_API void ServerUpdatePayload(const FARObjectUpdatePayload& NewPayload);
 	
 	UPROPERTY(BlueprintReadOnly, ReplicatedUsing=OnRep_Payload, Category = "AR Gameplay")
 	FARObjectUpdatePayload ReplicatedPayload;
 	
-	virtual void OnRep_Payload() override;
+	UE_API virtual void OnRep_Payload() override;
 };
 
 
-UCLASS(Blueprintable, BlueprintType, Category = "AR Gameplay")
-class AUGMENTEDREALITY_API UARMeshComponent : public UARComponent
+UCLASS(MinimalAPI, Blueprintable, BlueprintType, Category = "AR Gameplay")
+class UARMeshComponent : public UARComponent
 {
 	GENERATED_BODY()
 public:
 	/** Event when native representation is first added, called on server and clients. */
 	UFUNCTION(BlueprintImplementableEvent, meta = (DisplayName = "Add"))
-	void ReceiveAdd(const FARMeshUpdatePayload& Payload);
+	UE_API void ReceiveAdd(const FARMeshUpdatePayload& Payload);
 	
 	/** Event when native representation is updated, called on server and clients. */
 	UFUNCTION(BlueprintImplementableEvent, meta = (DisplayName = "Update"))
-	void ReceiveUpdate(const FARMeshUpdatePayload& Payload);
+	UE_API void ReceiveUpdate(const FARMeshUpdatePayload& Payload);
 	
 	/** Overridable native event for when native representation is updated. */
-	virtual void Update(UARTrackedGeometry* TrackedGeometry) override;
+	UE_API virtual void Update(UARTrackedGeometry* TrackedGeometry) override;
 	
 protected:
 	UFUNCTION(Reliable, Server, WithValidation)
-	void ServerUpdatePayload(const FARMeshUpdatePayload& NewPayload);
+	UE_API void ServerUpdatePayload(const FARMeshUpdatePayload& NewPayload);
 	
 	UPROPERTY(BlueprintReadOnly, ReplicatedUsing=OnRep_Payload, Category = "AR Gameplay")
 	FARMeshUpdatePayload ReplicatedPayload;
 	
-	virtual void OnRep_Payload() override;
-	virtual void UpdateVisualization_Implementation() override;
+	UE_API virtual void OnRep_Payload() override;
+	UE_API virtual void UpdateVisualization_Implementation() override;
 };
 
 
@@ -693,35 +695,37 @@ enum class EGeoAnchorComponentDebugMode : uint8
 	ShowGeoData = 1,
 };
 
-UCLASS(Blueprintable, BlueprintType, Category = "AR Gameplay")
-class AUGMENTEDREALITY_API UARGeoAnchorComponent : public UARComponent
+UCLASS(MinimalAPI, Blueprintable, BlueprintType, Category = "AR Gameplay")
+class UARGeoAnchorComponent : public UARComponent
 {
 	GENERATED_BODY()
 
 public:
 	UFUNCTION(BlueprintCallable, Category = "AR Gameplay | Debug", meta=(DevelopmentOnly))
-	static void SetGeoAnchorComponentDebugMode(EGeoAnchorComponentDebugMode NewDebugMode);
+	static UE_API void SetGeoAnchorComponentDebugMode(EGeoAnchorComponentDebugMode NewDebugMode);
 	
 	/** Event when native representation is first added, called on server and clients. */
 	UFUNCTION(BlueprintImplementableEvent, meta = (DisplayName = "Add"))
-	void ReceiveAdd(const FARGeoAnchorUpdatePayload& Payload);
+	UE_API void ReceiveAdd(const FARGeoAnchorUpdatePayload& Payload);
 	
 	/** Event when native representation is updated, called on server and clients. */
 	UFUNCTION(BlueprintImplementableEvent, meta = (DisplayName = "Update"))
-	void ReceiveUpdate(const FARGeoAnchorUpdatePayload& Payload);
+	UE_API void ReceiveUpdate(const FARGeoAnchorUpdatePayload& Payload);
 	
 	/** Overridable native event for when native representation is updated. */
-	virtual void Update(UARTrackedGeometry* TrackedGeometry) override;
+	UE_API virtual void Update(UARTrackedGeometry* TrackedGeometry) override;
 	
-	virtual void TickComponent(float DeltaTime, enum ELevelTick TickType, FActorComponentTickFunction* ThisTickFunction) override;
+	UE_API virtual void TickComponent(float DeltaTime, enum ELevelTick TickType, FActorComponentTickFunction* ThisTickFunction) override;
 	
 protected:
 	UFUNCTION(Reliable, Server, WithValidation)
-	void ServerUpdatePayload(const FARGeoAnchorUpdatePayload& NewPayload);
+	UE_API void ServerUpdatePayload(const FARGeoAnchorUpdatePayload& NewPayload);
 	
 	UPROPERTY(BlueprintReadOnly, ReplicatedUsing=OnRep_Payload, Category = "AR Gameplay")
 	FARGeoAnchorUpdatePayload ReplicatedPayload;
 	
-	virtual void OnRep_Payload() override;
-	static int32 GetDebugMode();
+	UE_API virtual void OnRep_Payload() override;
+	static UE_API int32 GetDebugMode();
 };
+
+#undef UE_API
diff --git a/Engine/Source/Runtime/AugmentedReality/Public/ARDependencyHandler.h b/Engine/Source/Runtime/AugmentedReality/Public/ARDependencyHandler.h
--- a/Engine/Source/Runtime/AugmentedReality/Public/ARDependencyHandler.h
+++ b/Engine/Source/Runtime/AugmentedReality/Public/ARDependencyHandler.h
@@ -13,6 +13,8 @@
 
 #include "ARDependencyHandler.generated.h"
 
+#define UE_API AUGMENTEDREALITY_API
+
 class UARSessionConfig;
 struct FFrame;
 
@@ -77,8 +79,8 @@ enum class EARServicePermissionRequestResult : uint8
  * 5. Call "UARBlueprintLibrary::StartARSession" to start the session.
  * Alternatively, you can also call "StartARSessionLatent" which handles dependency and permission internally.
  */
-UCLASS(BlueprintType, Abstract, Category="")
-class AUGMENTEDREALITY_API UARDependencyHandler : public UObject, public IModularFeature
+UCLASS(MinimalAPI, BlueprintType, Abstract, Category="")
+class UARDependencyHandler : public UObject, public IModularFeature
 {
 	GENERATED_BODY()
 	
@@ -94,7 +96,7 @@ public:
 	 * Can return null if the current platform doesn't support AR, or the AR system doesn't require dependency handling.
 	 */
 	UFUNCTION(BlueprintCallable, Category = "")
-	static UARDependencyHandler* GetARDependencyHandler();
+	static UE_API UARDependencyHandler* GetARDependencyHandler();
 	
 	/** Latent action to check AR availability on the current platform. */
 	UFUNCTION(BlueprintCallable, Category = "", meta = (Latent, LatentInfo = "LatentInfo", WorldContext = "WorldContextObject"))
@@ -115,3 +117,5 @@ public:
 	UFUNCTION(BlueprintCallable, Category = "", meta = (Latent, LatentInfo = "LatentInfo", WorldContext = "WorldContextObject"))
 	virtual void StartARSessionLatent(UObject* WorldContextObject, UARSessionConfig* SessionConfig, FLatentActionInfo LatentInfo) PURE_VIRTUAL(UARDependencyHandler::StartARSessionLatent, );
 };
+
+#undef UE_API
diff --git a/Engine/Source/Runtime/AugmentedReality/Public/ARGeoTrackingSupport.h b/Engine/Source/Runtime/AugmentedReality/Public/ARGeoTrackingSupport.h
--- a/Engine/Source/Runtime/AugmentedReality/Public/ARGeoTrackingSupport.h
+++ b/Engine/Source/Runtime/AugmentedReality/Public/ARGeoTrackingSupport.h
@@ -9,6 +9,8 @@
 
 #include "ARGeoTrackingSupport.generated.h"
 
+#define UE_API AUGMENTEDREALITY_API
+
 
 class FCheckGeoTrackingAvailabilityAsyncTask;
 class FGetGeoLocationAsyncTask;
@@ -83,8 +85,8 @@ enum class EARGeoTrackingAccuracy : uint8
 /**
  * Interface class for Geo tracking related features.
  */
-UCLASS(BlueprintType, Abstract, Category="AR|Geo Tracking")
-class AUGMENTEDREALITY_API UARGeoTrackingSupport : public UObject, public IModularFeature
+UCLASS(MinimalAPI, BlueprintType, Abstract, Category="AR|Geo Tracking")
+class UARGeoTrackingSupport : public UObject, public IModularFeature
 {
 	GENERATED_BODY()
 	
@@ -97,25 +99,25 @@ public:
 	
 	/** @return the interface object to support Geo tracking, return null on platforms don't support the feature. */
 	UFUNCTION(BlueprintCallable, Category = "AR AugmentedReality|Geo Tracking")
-	static class UARGeoTrackingSupport* GetGeoTrackingSupport();
+	static UE_API class UARGeoTrackingSupport* GetGeoTrackingSupport();
 	
 	/**
 	 * @return the async task to check Geo tracking availability at the user's current location.
 	 * Can return null if the feature is not supported.
 	 */
-	virtual TSharedPtr<FCheckGeoTrackingAvailabilityAsyncTask, ESPMode::ThreadSafe> CheckGeoTrackingAvailability(FString& OutError);
+	UE_API virtual TSharedPtr<FCheckGeoTrackingAvailabilityAsyncTask, ESPMode::ThreadSafe> CheckGeoTrackingAvailability(FString& OutError);
 	
 	/**
 	 * @return the async task to check Geo tracking availability at a specific location.
 	 * Can return null if the feature is not supported.
 	 */
-	virtual TSharedPtr<FCheckGeoTrackingAvailabilityAsyncTask, ESPMode::ThreadSafe> CheckGeoTrackingAvailability(float Longitude, float Latitude, FString& OutError);
+	UE_API virtual TSharedPtr<FCheckGeoTrackingAvailabilityAsyncTask, ESPMode::ThreadSafe> CheckGeoTrackingAvailability(float Longitude, float Latitude, FString& OutError);
 	
 	/**
 	 * @return the async task to convert a UE4 world position into a Geo location.
 	 * Can return null if the feature is not supported.
 	 */
-	virtual TSharedPtr<FGetGeoLocationAsyncTask, ESPMode::ThreadSafe> GetGeoLocationAtWorldPosition(const FVector& WorldPosition, FString& OutError);
+	UE_API virtual TSharedPtr<FGetGeoLocationAsyncTask, ESPMode::ThreadSafe> GetGeoLocationAtWorldPosition(const FVector& WorldPosition, FString& OutError);
 	
 	/**
 	 * @return the current session's Geo tracking state.
@@ -152,7 +154,7 @@ public:
 /**
  * Async task to check Geo tracking availability.
  */
-class AUGMENTEDREALITY_API FCheckGeoTrackingAvailabilityAsyncTask : public FARAsyncTask
+class FCheckGeoTrackingAvailabilityAsyncTask : public FARAsyncTask
 {
 public:
 	bool IsAvailable() const { return bIsAvailable; }
@@ -160,12 +162,12 @@ public:
 	/**
 	 * Finish the task with the availability result.
 	 */
-	void FinishWithAvailability(bool bInIsAvailable);
+	UE_API void FinishWithAvailability(bool bInIsAvailable);
 	
 	/**
 	 * Finish the task with an error.
 	 */
-	void FinishWithError(const FString& InError);
+	UE_API void FinishWithError(const FString& InError);
 	
 private:
 	bool bIsAvailable = false;
@@ -175,8 +177,8 @@ private:
 /**
  * Blueprint async task to check Geo tracking availability.
  */
-UCLASS()
-class AUGMENTEDREALITY_API UCheckGeoTrackingAvailabilityAsyncTaskBlueprintProxy : public UARBaseAsyncTaskBlueprintProxy
+UCLASS(MinimalAPI)
+class UCheckGeoTrackingAvailabilityAsyncTaskBlueprintProxy : public UARBaseAsyncTaskBlueprintProxy
 {
 	GENERATED_BODY()
 	
@@ -193,19 +195,19 @@ public:
 	 * Check Geo tracking availability at the user's current location.
 	 */
 	UFUNCTION(BlueprintCallable, Category = "AR AugmentedReality|Geo Tracking", meta=(BlueprintInternalUseOnly="true", WorldContext = "WorldContextObject"))
-	static UCheckGeoTrackingAvailabilityAsyncTaskBlueprintProxy* CheckGeoTrackingAvailability(UObject* WorldContextObject);
+	static UE_API UCheckGeoTrackingAvailabilityAsyncTaskBlueprintProxy* CheckGeoTrackingAvailability(UObject* WorldContextObject);
 	
 	/**
 	 * Check Geo tracking availability at a specific Geo location.
 	 */
 	UFUNCTION(BlueprintCallable, Category = "AR AugmentedReality|Geo Tracking", meta=(BlueprintInternalUseOnly="true", WorldContext = "WorldContextObject"))
-	static UCheckGeoTrackingAvailabilityAsyncTaskBlueprintProxy* CheckGeoTrackingAvailabilityAtLocation(UObject* WorldContextObject, float Longitude, float Latitude);
+	static UE_API UCheckGeoTrackingAvailabilityAsyncTaskBlueprintProxy* CheckGeoTrackingAvailabilityAtLocation(UObject* WorldContextObject, float Longitude, float Latitude);
 	
 private:
 	// UBlueprintAsyncActionBase interface
-	void Activate() override;
-	void ReportSuccess() override;
-	void ReportFailure() override;
+	UE_API void Activate() override;
+	UE_API void ReportSuccess() override;
+	UE_API void ReportFailure() override;
 	//~UBlueprintAsyncActionBase interface
 	
 	TOptional<float> Longitude;
@@ -221,18 +223,18 @@ private:
 /**
  * Async task to convert Geo location.
  */
-class AUGMENTEDREALITY_API FGetGeoLocationAsyncTask : public FARAsyncTask
+class FGetGeoLocationAsyncTask : public FARAsyncTask
 {
 public:
 	/**
 	 * Finish the task with a Geo location.
 	 */
-	void FinishWithGeoLocation(float InLongitude, float InLatitude, float InAltitude);
+	UE_API void FinishWithGeoLocation(float InLongitude, float InLatitude, float InAltitude);
 	
 	/**
 	 * Finish the task with an error.
 	 */
-	void FinishWithError(const FString& InError);
+	UE_API void FinishWithError(const FString& InError);
 	
 	/**
 	 * @return the longitude of the converted Geo location.
@@ -259,8 +261,8 @@ private:
 /**
  * Blueprint async task to convert Geo location.
  */
-UCLASS()
-class AUGMENTEDREALITY_API UGetGeoLocationAsyncTaskBlueprintProxy : public UARBaseAsyncTaskBlueprintProxy
+UCLASS(MinimalAPI)
+class UGetGeoLocationAsyncTaskBlueprintProxy : public UARBaseAsyncTaskBlueprintProxy
 {
 	GENERATED_BODY()
 	
@@ -277,13 +279,13 @@ public:
 	 * Convert a position in UE4 world space into a Geo location.
 	 */
 	UFUNCTION(BlueprintCallable, Category = "AR AugmentedReality|Geo Tracking", meta=(BlueprintInternalUseOnly="true", WorldContext = "WorldContextObject"))
-	static UGetGeoLocationAsyncTaskBlueprintProxy* GetGeoLocationAtWorldPosition(UObject* WorldContextObject, const FVector& WorldPosition);
+	static UE_API UGetGeoLocationAsyncTaskBlueprintProxy* GetGeoLocationAtWorldPosition(UObject* WorldContextObject, const FVector& WorldPosition);
 	
 private:
 	// UBlueprintAsyncActionBase interface
-	void Activate() override;
-	void ReportSuccess() override;
-	void ReportFailure() override;
+	UE_API void Activate() override;
+	UE_API void ReportSuccess() override;
+	UE_API void ReportFailure() override;
 	//~UBlueprintAsyncActionBase interface
 	
 	FVector WorldPosition = FVector::ZeroVector;
@@ -293,3 +295,5 @@ private:
 	
 	FString Error;
 };
+
+#undef UE_API
diff --git a/Engine/Source/Runtime/AugmentedReality/Public/ARLifeCycleComponent.h b/Engine/Source/Runtime/AugmentedReality/Public/ARLifeCycleComponent.h
--- a/Engine/Source/Runtime/AugmentedReality/Public/ARLifeCycleComponent.h
+++ b/Engine/Source/Runtime/AugmentedReality/Public/ARLifeCycleComponent.h
@@ -9,25 +9,27 @@
 #include "ARActor.h"
 #include "ARLifeCycleComponent.generated.h"
 
+#define UE_API AUGMENTEDREALITY_API
+
 class USceneComponent;
 
-UCLASS(BlueprintType, Experimental, meta = (BlueprintSpawnableComponent), ClassGroup = "AR Gameplay")
-class AUGMENTEDREALITY_API UARLifeCycleComponent : public USceneComponent
+UCLASS(MinimalAPI, BlueprintType, Experimental, meta = (BlueprintSpawnableComponent), ClassGroup = "AR Gameplay")
+class UARLifeCycleComponent : public USceneComponent
 {
 	GENERATED_BODY()
 
-	virtual void OnComponentCreated() override;
-	virtual void DestroyComponent(bool bPromoteChildren = false) override;
+	UE_API virtual void OnComponentCreated() override;
+	UE_API virtual void DestroyComponent(bool bPromoteChildren = false) override;
 
 public:
 	DECLARE_MULTICAST_DELEGATE_TwoParams(FRequestSpawnARActorDelegate, UClass*, FGuid);
-	static FRequestSpawnARActorDelegate RequestSpawnARActorDelegate;
+	static UE_API FRequestSpawnARActorDelegate RequestSpawnARActorDelegate;
 	
 	DECLARE_MULTICAST_DELEGATE_ThreeParams(FOnSpawnARActorDelegate, AARActor*, UARComponent*, FGuid);
-	static FOnSpawnARActorDelegate OnSpawnARActorDelegate;
+	static UE_API FOnSpawnARActorDelegate OnSpawnARActorDelegate;
 	
 	DECLARE_MULTICAST_DELEGATE_OneParam(FRequestDestroyARActorDelegate, AARActor*);
-	static FRequestDestroyARActorDelegate RequestDestroyARActorDelegate;
+	static UE_API FRequestDestroyARActorDelegate RequestDestroyARActorDelegate;
 	
 	/** Called when an AR actor is spawned on the server */
 	DECLARE_DYNAMIC_MULTICAST_DELEGATE_ThreeParams(FInstanceARActorSpawnedDelegate, UClass*, ComponentClass, FGuid, NativeID, AARActor*, SpawnedActor);
@@ -40,18 +42,20 @@ public:
 	FInstanceARActorToBeDestroyedDelegate OnARActorToBeDestroyedDelegate;
 	
 protected:
-	virtual void OnUnregister() override;
+	UE_API virtual void OnUnregister() override;
 	
 private:
-	void CallInstanceRequestSpawnARActorDelegate(UClass* Class, FGuid NativeID);
-	void CallInstanceRequestDestroyARActorDelegate(AARActor* Actor);
+	UE_API void CallInstanceRequestSpawnARActorDelegate(UClass* Class, FGuid NativeID);
+	UE_API void CallInstanceRequestDestroyARActorDelegate(AARActor* Actor);
 	
 	UFUNCTION(Reliable, Server, WithValidation)
-	void ServerSpawnARActor(UClass* ComponentClass, FGuid NativeID);
+	UE_API void ServerSpawnARActor(UClass* ComponentClass, FGuid NativeID);
 	
 	UFUNCTION(Reliable, Server, WithValidation)
-	void ServerDestroyARActor(AARActor* Actor);
+	UE_API void ServerDestroyARActor(AARActor* Actor);
 	
 	FDelegateHandle SpawnDelegateHandle;
 	FDelegateHandle DestroyDelegateHandle;
 };
+
+#undef UE_API
diff --git a/Engine/Source/Runtime/AugmentedReality/Public/ARLightEstimate.h b/Engine/Source/Runtime/AugmentedReality/Public/ARLightEstimate.h
--- a/Engine/Source/Runtime/AugmentedReality/Public/ARLightEstimate.h
+++ b/Engine/Source/Runtime/AugmentedReality/Public/ARLightEstimate.h
@@ -10,6 +10,8 @@
 
 #include "ARLightEstimate.generated.h"
 
+#define UE_API AUGMENTEDREALITY_API
+
 struct FFrame;
 
 
@@ -19,26 +21,26 @@ class UARLightEstimate : public UObject
 	GENERATED_BODY()
 };
 
-UCLASS(BlueprintType, Category = "AR AugmentedReality|Light Estimation")
-class AUGMENTEDREALITY_API UARBasicLightEstimate : public UARLightEstimate
+UCLASS(MinimalAPI, BlueprintType, Category = "AR AugmentedReality|Light Estimation")
+class UARBasicLightEstimate : public UARLightEstimate
 {
 	GENERATED_BODY()
 	
 public:
-	void SetLightEstimate(float InAmbientIntensityLumens, float InColorTemperatureKelvin);
+	UE_API void SetLightEstimate(float InAmbientIntensityLumens, float InColorTemperatureKelvin);
 	
-	void SetLightEstimate(FVector InRGBScaleFactor, float InPixelIntensity);
+	UE_API void SetLightEstimate(FVector InRGBScaleFactor, float InPixelIntensity);
 
-	void SetLightEstimate(float InColorTemperatureKelvin, FLinearColor InAmbientColor);
+	UE_API void SetLightEstimate(float InColorTemperatureKelvin, FLinearColor InAmbientColor);
 
 	UFUNCTION(BlueprintPure, Category = "AR AugmentedReality|Light Estimation")
-	float GetAmbientIntensityLumens() const;
+	UE_API float GetAmbientIntensityLumens() const;
 	
 	UFUNCTION(BlueprintPure, Category = "AR AugmentedReality|Light Estimation")
-	float GetAmbientColorTemperatureKelvin() const;
+	UE_API float GetAmbientColorTemperatureKelvin() const;
 	
 	UFUNCTION(BlueprintPure, Category = "AR AugmentedReality|Light Estimation")
-	FLinearColor GetAmbientColor() const;
+	UE_API FLinearColor GetAmbientColor() const;
 	
 private:
 	UPROPERTY()
@@ -50,3 +52,5 @@ private:
 	UPROPERTY()
 	FLinearColor AmbientColor;
 };
+
+#undef UE_API
diff --git a/Engine/Source/Runtime/AugmentedReality/Public/AROriginActor.h b/Engine/Source/Runtime/AugmentedReality/Public/AROriginActor.h
--- a/Engine/Source/Runtime/AugmentedReality/Public/AROriginActor.h
+++ b/Engine/Source/Runtime/AugmentedReality/Public/AROriginActor.h
@@ -6,16 +6,20 @@
 
 #include "AROriginActor.generated.h"
 
+#define UE_API AUGMENTEDREALITY_API
+
 /**
  * Simple actor that is spawned at the origin for AR systems that want to hang components on an actor
  * Spawned as a custom class for easier TObjectIterator use
  */
-UCLASS(BlueprintType)
-class AUGMENTEDREALITY_API AAROriginActor :
+UCLASS(MinimalAPI, BlueprintType)
+class AAROriginActor :
 	public AActor
 {
 	GENERATED_UCLASS_BODY()
 public:
 	/** Used by the AR system to get the origin actor for the current world */
-	static AAROriginActor* GetOriginActor();
+	static UE_API AAROriginActor* GetOriginActor();
 };
+
+#undef UE_API
diff --git a/Engine/Source/Runtime/AugmentedReality/Public/ARPin.h b/Engine/Source/Runtime/AugmentedReality/Public/ARPin.h
--- a/Engine/Source/Runtime/AugmentedReality/Public/ARPin.h
+++ b/Engine/Source/Runtime/AugmentedReality/Public/ARPin.h
@@ -5,16 +5,18 @@
 #include "ARTypes.h"
 #include "ARPin.generated.h"
 
+#define UE_API AUGMENTEDREALITY_API
+
 class FARSupportInterface ;
 class USceneComponent;
 
-UCLASS(BlueprintType, Experimental, Category="AR AugmentedReality")
-class AUGMENTEDREALITY_API UARPin : public UObject
+UCLASS(MinimalAPI, BlueprintType, Experimental, Category="AR AugmentedReality")
+class UARPin : public UObject
 {
 	GENERATED_BODY()
 	
 public:
-	virtual void InitARPin( const TSharedRef<FARSupportInterface , ESPMode::ThreadSafe>& InTrackingSystemOwner, USceneComponent* InComponentToPin, const FTransform& InLocalToTrackingTransform, UARTrackedGeometry* InTrackedGeometry, const FName InDebugName );
+	UE_API virtual void InitARPin( const TSharedRef<FARSupportInterface , ESPMode::ThreadSafe>& InTrackingSystemOwner, USceneComponent* InComponentToPin, const FTransform& InLocalToTrackingTransform, UARTrackedGeometry* InTrackedGeometry, const FName InDebugName );
 
 	/**
 	 * Maps from a Pin's Local Space to the Tracking Space.
@@ -22,7 +24,7 @@ public:
 	 * yield the Pin's position in Tracking Space.
 	 */
 	UFUNCTION(BlueprintPure, Category="AR AugmentedReality|Pin")
-	FTransform GetLocalToTrackingTransform() const;
+	UE_API FTransform GetLocalToTrackingTransform() const;
 	
 	
 	/**
@@ -30,33 +32,33 @@ public:
 	 * appends the TrackingToWorld Transform.
 	 */
 	UFUNCTION(BlueprintPure, Category="AR AugmentedReality|Pin")
-	FTransform GetLocalToWorldTransform() const;
+	UE_API FTransform GetLocalToWorldTransform() const;
 
 	/**
 	 * Return the current tracking state of this Pin.
 	 */
 	UFUNCTION(BlueprintPure, Category = "AR AugmentedReality|Pin")
-	EARTrackingState GetTrackingState() const;
+	UE_API EARTrackingState GetTrackingState() const;
 	
 	/**
 	 * The TrackedGeometry (if any) that this this pin is being "stuck" into.
 	 */
 	UFUNCTION(BlueprintPure, Category="AR AugmentedReality|Pin")
-	UARTrackedGeometry* GetTrackedGeometry() const;
+	UE_API UARTrackedGeometry* GetTrackedGeometry() const;
 	
 	/** @return the PinnedComponent that this UARPin is pinning to the TrackedGeometry */
 	UFUNCTION(BlueprintPure, Category="AR AugmentedReality|Pin")
-	USceneComponent* GetPinnedComponent() const;
+	UE_API USceneComponent* GetPinnedComponent() const;
 	
 	UFUNCTION()
-	virtual void DebugDraw( UWorld* World, const FLinearColor& Color, float Scale = 5.0f, float PersistForSeconds = 0.0f) const;
+	UE_API virtual void DebugDraw( UWorld* World, const FLinearColor& Color, float Scale = 5.0f, float PersistForSeconds = 0.0f) const;
 	
 	UFUNCTION(BlueprintPure, Category="AR AugmentedReality|Pin")
-	FName GetDebugName() const;
+	UE_API FName GetDebugName() const;
 	
-	void SetOnARTrackingStateChanged( const FOnARTrackingStateChanged& InHandler );
+	UE_API void SetOnARTrackingStateChanged( const FOnARTrackingStateChanged& InHandler );
 	
-	void SetOnARTransformUpdated( const FOnARTransformUpdated& InHandler );
+	UE_API void SetOnARTransformUpdated( const FOnARTransformUpdated& InHandler );
 
 	void SetNativeResource(void* InNativeResource)
 	{
@@ -70,16 +72,16 @@ public:
 	
 public:
 	
-	FTransform GetLocalToTrackingTransform_NoAlignment() const;
+	UE_API FTransform GetLocalToTrackingTransform_NoAlignment() const;
 	
 	/** Notify the ARPin about changes to how it is being tracked. */
-	void OnTrackingStateChanged(EARTrackingState NewTrackingState);
+	UE_API void OnTrackingStateChanged(EARTrackingState NewTrackingState);
 	
 	/** Notify this UARPin that the transform of the Pin has changed */
-	void OnTransformUpdated(const FTransform& NewLocalToTrackingTransform);
+	UE_API void OnTransformUpdated(const FTransform& NewLocalToTrackingTransform);
 	
 	/** Notify the UARPin that the AlignmentTransform has changing. */
-	void UpdateAlignmentTransform( const FTransform& NewAlignmentTransform );
+	UE_API void UpdateAlignmentTransform( const FTransform& NewAlignmentTransform );
 	
 	void SetPinnedComponent(USceneComponent* InComponentToPin)
 	{
@@ -87,10 +89,10 @@ public:
 	}
 	
 protected:
-	TSharedPtr<FARSupportInterface , ESPMode::ThreadSafe> GetARSystem() const;
+	UE_API TSharedPtr<FARSupportInterface , ESPMode::ThreadSafe> GetARSystem() const;
 	
 private:
-	static uint32 DebugPinId;
+	static UE_API uint32 DebugPinId;
 	
 	UPROPERTY()
 	TObjectPtr<UARTrackedGeometry> TrackedGeometry;
@@ -120,3 +122,5 @@ private:
 	// The native resource pointer on the AR platform.
 	void* NativeResource;
 };
+
+#undef UE_API
diff --git a/Engine/Source/Runtime/AugmentedReality/Public/ARSessionConfig.h b/Engine/Source/Runtime/AugmentedReality/Public/ARSessionConfig.h
--- a/Engine/Source/Runtime/AugmentedReality/Public/ARSessionConfig.h
+++ b/Engine/Source/Runtime/AugmentedReality/Public/ARSessionConfig.h
@@ -10,6 +10,8 @@
 
 #include "ARSessionConfig.generated.h"
 
+#define UE_API AUGMENTEDREALITY_API
+
 class UMaterialInterface;
 
 /** Options for how the scene’s coordinate system is constructed. This feature is used by ARKit. */
@@ -176,149 +178,149 @@ enum class EARSceneReconstruction : uint8
 /**
  * An Unreal Data Asset that defines what features are used in the AR session.
  */
-UCLASS(BlueprintType, Category="AR Settings")
-class AUGMENTEDREALITY_API UARSessionConfig : public UDataAsset
+UCLASS(MinimalAPI, BlueprintType, Category="AR Settings")
+class UARSessionConfig : public UDataAsset
 {
 	GENERATED_BODY()
 
 public:
 	/** The constructor for the AR Session Config Data Asset. */
-	UARSessionConfig();
+	UE_API UARSessionConfig();
 	
 public:
 	/** @see EARWorldAlignment */
 	UFUNCTION(BlueprintCallable, Category = "AR Settings")
-	EARWorldAlignment GetWorldAlignment() const;
+	UE_API EARWorldAlignment GetWorldAlignment() const;
 
 	/** @see SessionType */
 	UFUNCTION(BlueprintCallable, Category = "AR Settings")
-	EARSessionType GetSessionType() const;
+	UE_API EARSessionType GetSessionType() const;
 
 	/** @see PlaneDetectionMode */
 	UFUNCTION(BlueprintCallable, Category = "AR Settings")
-	EARPlaneDetectionMode GetPlaneDetectionMode() const;
+	UE_API EARPlaneDetectionMode GetPlaneDetectionMode() const;
 
 	/** @see LightEstimationMode */
 	UFUNCTION(BlueprintCallable, Category = "AR Settings")
-	EARLightEstimationMode GetLightEstimationMode() const;
+	UE_API EARLightEstimationMode GetLightEstimationMode() const;
 
 	/** @see FrameSyncMode */
 	UFUNCTION(BlueprintCallable, Category = "AR Settings")
-	EARFrameSyncMode GetFrameSyncMode() const;
+	UE_API EARFrameSyncMode GetFrameSyncMode() const;
 
 	/** @see bEnableAutomaticCameraOverlay */
 	UFUNCTION(BlueprintCallable, Category = "AR Settings")
-	bool ShouldRenderCameraOverlay() const;
+	UE_API bool ShouldRenderCameraOverlay() const;
 
 	/** @see bEnableAutomaticCameraTracking */
 	UFUNCTION(BlueprintCallable, Category = "AR Settings")
-	bool ShouldEnableCameraTracking() const;
+	UE_API bool ShouldEnableCameraTracking() const;
 
 	/** @see bEnableAutoFocus */
 	UFUNCTION(BlueprintCallable, Category = "AR Settings")
-	bool ShouldEnableAutoFocus() const;
+	UE_API bool ShouldEnableAutoFocus() const;
 
 	/** @see bEnableAutoFocus */
 	UFUNCTION(BlueprintCallable, Category = "AR Settings")
-	void SetEnableAutoFocus(bool bNewValue);
+	UE_API void SetEnableAutoFocus(bool bNewValue);
 
 	/** @see bResetCameraTracking */
 	UFUNCTION(BlueprintCallable, Category = "AR Settings")
-	bool ShouldResetCameraTracking() const;
+	UE_API bool ShouldResetCameraTracking() const;
 	
 	/** @see bResetCameraTracking */
 	UFUNCTION(BlueprintCallable, Category = "AR Settings")
-	void SetResetCameraTracking(bool bNewValue);
+	UE_API void SetResetCameraTracking(bool bNewValue);
 
 	/** @see bResetTrackedObjects */
 	UFUNCTION(BlueprintCallable, Category = "AR Settings")
-	bool ShouldResetTrackedObjects() const;
+	UE_API bool ShouldResetTrackedObjects() const;
 
 	/** @see bResetTrackedObjects */
 	UFUNCTION(BlueprintCallable, Category = "AR Settings")
-	void SetResetTrackedObjects(bool bNewValue);
+	UE_API void SetResetTrackedObjects(bool bNewValue);
 	
 	/** @see CandidateImages */
 	UFUNCTION(BlueprintCallable, Category = "AR Settings")
-	const TArray<UARCandidateImage*>& GetCandidateImageList() const;
+	UE_API const TArray<UARCandidateImage*>& GetCandidateImageList() const;
 
 	// Add a new CandidateImage to the ARSessionConfig.
 	UFUNCTION(BlueprintCallable, Category = "AR Settings")
-	void AddCandidateImage(UARCandidateImage* NewCandidateImage);
+	UE_API void AddCandidateImage(UARCandidateImage* NewCandidateImage);
     
 	// Remove a candidate image from the ARSessionConfig, by pointer, note the image object must match, not the content of the image.
 	UFUNCTION(BlueprintCallable, Category = "AR Settings")
-	void RemoveCandidateImage(UARCandidateImage* CandidateImage);
+	UE_API void RemoveCandidateImage(UARCandidateImage* CandidateImage);
 
 	// Remove a candidate image from the ARSessionConfig, by index.
 	UFUNCTION(BlueprintCallable, Category = "AR Settings")
-	void RemoveCandidateImageAtIndex(int Index);
+	UE_API void RemoveCandidateImageAtIndex(int Index);
 
 	// Remove all candidate images from the ARSessionConfig
 	UFUNCTION(BlueprintCallable, Category = "AR Settings")
-	void ClearCandidateImages();
+	UE_API void ClearCandidateImages();
 
 	/** @see MaxNumSimultaneousImagesTracked */
 	UFUNCTION(BlueprintCallable, Category = "AR Settings")
-    int32 GetMaxNumSimultaneousImagesTracked() const;
+    UE_API int32 GetMaxNumSimultaneousImagesTracked() const;
 	
 	/** @see EnvironmentCaptureProbeType */
 	UFUNCTION(BlueprintCallable, Category = "AR Settings")
-	EAREnvironmentCaptureProbeType GetEnvironmentCaptureProbeType() const;
+	UE_API EAREnvironmentCaptureProbeType GetEnvironmentCaptureProbeType() const;
 	
 	/** @see WorldMapData */
 	UFUNCTION(BlueprintCallable, Category = "AR Settings")
-	const TArray<uint8>& GetWorldMapData() const;
+	UE_API const TArray<uint8>& GetWorldMapData() const;
 	/** @see WorldMapData */
 	UFUNCTION(BlueprintCallable, Category = "AR Settings")
-	void SetWorldMapData(TArray<uint8> WorldMapData);
+	UE_API void SetWorldMapData(TArray<uint8> WorldMapData);
 
 	/** @see CandidateObjects */
 	UFUNCTION(BlueprintCallable, Category = "AR Settings")
-	const TArray<UARCandidateObject*>& GetCandidateObjectList() const;
+	UE_API const TArray<UARCandidateObject*>& GetCandidateObjectList() const;
 	/** @see CandidateObjects */
 	UFUNCTION(BlueprintCallable, Category = "AR Settings")
-	void SetCandidateObjectList(const TArray<UARCandidateObject*>& InCandidateObjects);
+	UE_API void SetCandidateObjectList(const TArray<UARCandidateObject*>& InCandidateObjects);
 	/** @see CandidateObjects */
 	UFUNCTION(BlueprintCallable, Category = "AR Settings")
-	void AddCandidateObject(UARCandidateObject* CandidateObject);
+	UE_API void AddCandidateObject(UARCandidateObject* CandidateObject);
 
 	/** @see DesiredVideoFormat */
 	UFUNCTION(BlueprintCallable, Category = "AR Settings")
-	FARVideoFormat GetDesiredVideoFormat() const;
+	UE_API FARVideoFormat GetDesiredVideoFormat() const;
 	/** @see DesiredVideoFormat */
 	UFUNCTION(BlueprintCallable, Category = "AR Settings")
-	void SetDesiredVideoFormat(FARVideoFormat NewFormat);
+	UE_API void SetDesiredVideoFormat(FARVideoFormat NewFormat);
 
 	/** @see FaceTrackingDirection */
 	UFUNCTION(BlueprintCallable, Category = "AR Settings")
-	EARFaceTrackingDirection GetFaceTrackingDirection() const;
+	UE_API EARFaceTrackingDirection GetFaceTrackingDirection() const;
 	/** @see FaceTrackingDirection */
 	UFUNCTION(BlueprintCallable, Category = "AR Settings")
-	void SetFaceTrackingDirection(EARFaceTrackingDirection InDirection);
+	UE_API void SetFaceTrackingDirection(EARFaceTrackingDirection InDirection);
 	
 	/** @see FaceTrackingUpdate */
 	UFUNCTION(BlueprintCallable, Category = "AR Settings")
-	EARFaceTrackingUpdate GetFaceTrackingUpdate() const;
+	UE_API EARFaceTrackingUpdate GetFaceTrackingUpdate() const;
 	/** @see FaceTrackingUpdate */
 	UFUNCTION(BlueprintCallable, Category = "AR Settings")
-	void SetFaceTrackingUpdate(EARFaceTrackingUpdate InUpdate);
+	UE_API void SetFaceTrackingUpdate(EARFaceTrackingUpdate InUpdate);
 
 	/** @see EnabledSessionTrackingFeatures */
 	UFUNCTION(BlueprintCallable, Category = "AR Settings")
-	EARSessionTrackingFeature GetEnabledSessionTrackingFeature() const;
+	UE_API EARSessionTrackingFeature GetEnabledSessionTrackingFeature() const;
 	
 	/** @see SceneReconstructionMethod */
 	UFUNCTION(BlueprintCallable, Category = "AR Settings")
-	EARSceneReconstruction GetSceneReconstructionMethod() const;
+	UE_API EARSceneReconstruction GetSceneReconstructionMethod() const;
 	
 	/** @see EnabledSessionTrackingFeatures */
 	UFUNCTION(BlueprintCallable, Category = "AR Settings")
-	void SetSessionTrackingFeatureToEnable(EARSessionTrackingFeature InSessionTrackingFeature);
+	UE_API void SetSessionTrackingFeatureToEnable(EARSessionTrackingFeature InSessionTrackingFeature);
 	
 	/** @see SceneReconstructionMethod */
 	UFUNCTION(BlueprintCallable, Category = "AR Settings")
-	void SetSceneReconstructionMethod(EARSceneReconstruction InSceneReconstructionMethod);
+	UE_API void SetSceneReconstructionMethod(EARSceneReconstruction InSceneReconstructionMethod);
 
 	/** @see bHorizontalPlaneDetection */
 	bool ShouldDoHorizontalPlaneDetection() const { return bHorizontalPlaneDetection; }
@@ -326,28 +328,28 @@ public:
 	bool ShouldDoVerticalPlaneDetection() const { return bVerticalPlaneDetection; }
 	
 	/** @see  SerializedARCandidateImageDatabase */
-	const TArray<uint8>& GetSerializedARCandidateImageDatabase() const;
+	UE_API const TArray<uint8>& GetSerializedARCandidateImageDatabase() const;
 
 	/** @see PlaneComponentClass */
-	UClass* GetPlaneComponentClass(void) const;
+	UE_API UClass* GetPlaneComponentClass(void) const;
 	/** @see PointComponentClass */
-	UClass* GetPointComponentClass(void) const;
+	UE_API UClass* GetPointComponentClass(void) const;
 	/** @see FaceComponentClass */
-	UClass* GetFaceComponentClass(void) const;
+	UE_API UClass* GetFaceComponentClass(void) const;
 	/** @see ImageComponentClass */
-	UClass* GetImageComponentClass(void) const;
+	UE_API UClass* GetImageComponentClass(void) const;
 	/** @see QRCodeComponentClass */
-	UClass* GetQRCodeComponentClass(void) const;
+	UE_API UClass* GetQRCodeComponentClass(void) const;
 	/** @see PoseComponentClass */
-	UClass* GetPoseComponentClass(void) const;
+	UE_API UClass* GetPoseComponentClass(void) const;
 	/** @see EnvironmentProbeComponentClass */
-	UClass* GetEnvironmentProbeComponentClass(void) const;
+	UE_API UClass* GetEnvironmentProbeComponentClass(void) const;
 	/** @see ObjectComponentClass */
-	UClass* GetObjectComponentClass(void) const;
+	UE_API UClass* GetObjectComponentClass(void) const;
 	/** @see MeshComponentClass */
-	UClass* GetMeshComponentClass(void) const;
+	UE_API UClass* GetMeshComponentClass(void) const;
 	/** @see GeoAnchorComponentClass */
-	UClass* GetGeoAnchorComponentClass(void) const;
+	UE_API UClass* GetGeoAnchorComponentClass(void) const;
 	
 	/** @see DefaultMeshMaterial */
 	UMaterialInterface* GetDefaultMeshMaterial() const { return DefaultMeshMaterial; }
@@ -398,11 +400,11 @@ public:
 	bool bUseStandardOnboardingUX = false;
 	
 	/** @see bUseOptimalVideoFormat */
-	bool ShouldUseOptimalVideoFormat() const;
+	UE_API bool ShouldUseOptimalVideoFormat() const;
 
 private:
 	//~ UObject interface
-	virtual void Serialize(FArchive& Ar) override;
+	UE_API virtual void Serialize(FArchive& Ar) override;
 	//~ UObject interface
 
 protected:
@@ -557,3 +559,5 @@ protected:
 	UPROPERTY(EditAnywhere, Category = "AR Settings | World Mapping")
 	TObjectPtr<UMaterialInterface> DefaultWireframeMeshMaterial;
 };
+
+#undef UE_API
diff --git a/Engine/Source/Runtime/AugmentedReality/Public/ARSessionConfigCookSupport.h b/Engine/Source/Runtime/AugmentedReality/Public/ARSessionConfigCookSupport.h
--- a/Engine/Source/Runtime/AugmentedReality/Public/ARSessionConfigCookSupport.h
+++ b/Engine/Source/Runtime/AugmentedReality/Public/ARSessionConfigCookSupport.h
@@ -8,7 +8,7 @@
 #include "Serialization/Archive.h"
 #include "Containers/UnrealString.h"
 
-class AUGMENTEDREALITY_API IARSessionConfigCookSupport : public IModularFeature
+class IARSessionConfigCookSupport : public IModularFeature
 {
 public:
 	static FName GetModularFeatureName()
diff --git a/Engine/Source/Runtime/AugmentedReality/Public/ARSharedWorldGameMode.h b/Engine/Source/Runtime/AugmentedReality/Public/ARSharedWorldGameMode.h
--- a/Engine/Source/Runtime/AugmentedReality/Public/ARSharedWorldGameMode.h
+++ b/Engine/Source/Runtime/AugmentedReality/Public/ARSharedWorldGameMode.h
@@ -9,9 +9,11 @@
 
 #include "ARSharedWorldGameMode.generated.h"
 
+#define UE_API AUGMENTEDREALITY_API
+
 /** Per player information about what data has been sent to them */
 USTRUCT(BlueprintType)
-struct AUGMENTEDREALITY_API FARSharedWorldReplicationState
+struct FARSharedWorldReplicationState
 {
 	GENERATED_BODY()
 	
@@ -29,8 +31,8 @@ struct AUGMENTEDREALITY_API FARSharedWorldReplicationState
 	int32 ARWorldOffset;
 };
 
-UCLASS(BlueprintType)
-class AUGMENTEDREALITY_API AARSharedWorldGameMode :
+UCLASS(MinimalAPI, BlueprintType)
+class AARSharedWorldGameMode :
 	public AGameMode
 {
 	GENERATED_UCLASS_BODY()
@@ -46,7 +48,7 @@ public:
 	 * @param ImageData the blob to use as the image data
 	 */
 	UFUNCTION(BlueprintCallable, BlueprintAuthorityOnly, meta=(DisplayName="Set AR Preview Image Data"), Category="AR Shared World")
-	void SetPreviewImageData(TArray<uint8> ImageData);
+	UE_API void SetPreviewImageData(TArray<uint8> ImageData);
 	
 	/**
 	 * Sets the image data for the shared world game session
@@ -54,23 +56,23 @@ public:
 	 * @param ARWorldData the blob to use as the AR world data
 	 */
 	UFUNCTION(BlueprintCallable, BlueprintAuthorityOnly, meta=(DisplayName="Set AR Shared World Data"), Category="AR Shared World")
-	void SetARSharedWorldData(TArray<uint8> ARWorldData);
+	UE_API void SetARSharedWorldData(TArray<uint8> ARWorldData);
 	
 	/**
 	 * Tells the game mode that the AR data is ready and should be replicated to all connected clients
 	 */
 	UFUNCTION(BlueprintCallable, BlueprintAuthorityOnly, meta=(DisplayName="Set AR World Sharing Is Ready"), Category="AR Shared World")
-	void SetARWorldSharingIsReady();
+	UE_API void SetARWorldSharingIsReady();
 	
 	/**
 	 * @return the game state for this game mode
 	 */
 	UFUNCTION(BlueprintCallable, BlueprintAuthorityOnly, Category="AR Shared World")
-	AARSharedWorldGameState* GetARSharedWorldGameState();
+	UE_API AARSharedWorldGameState* GetARSharedWorldGameState();
 	
 private:
-	virtual void Tick(float DeltaSeconds) override;
-	virtual void Logout(AController* Exiting) override;
+	UE_API virtual void Tick(float DeltaSeconds) override;
+	UE_API virtual void Logout(AController* Exiting) override;
 
 	/** Tracks whether the data should be sent to all clients or not */
 	bool bShouldSendSharedWorldData;
@@ -81,3 +83,5 @@ private:
 	/** Scratch buffer for sending chunks, cached to avoid allocation churn */
 	TArray<uint8> SendBuffer;
 };
+
+#undef UE_API
diff --git a/Engine/Source/Runtime/AugmentedReality/Public/ARSharedWorldGameState.h b/Engine/Source/Runtime/AugmentedReality/Public/ARSharedWorldGameState.h
--- a/Engine/Source/Runtime/AugmentedReality/Public/ARSharedWorldGameState.h
+++ b/Engine/Source/Runtime/AugmentedReality/Public/ARSharedWorldGameState.h
@@ -7,8 +7,10 @@
 #include "GameFramework/GameState.h"
 #include "ARSharedWorldGameState.generated.h"
 
-UCLASS(BlueprintType)
-class AUGMENTEDREALITY_API AARSharedWorldGameState :
+#define UE_API AUGMENTEDREALITY_API
+
+UCLASS(MinimalAPI, BlueprintType)
+class AARSharedWorldGameState :
 	public AGameState
 {
 	GENERATED_UCLASS_BODY()
@@ -45,7 +47,7 @@ public:
 	 * @param PreviewImageSize the total size in bytes of the image data that will be sent
 	 * @param ARWorldDataSize the total size in bytes of the AR world data that will be sent
 	 */
-	void InitSharedWorld(int32 PreviewImageSize, int32 ARWorldDataSize);
+	UE_API void InitSharedWorld(int32 PreviewImageSize, int32 ARWorldDataSize);
 	
 	/**
 	 * Copies the buffer into the image data (client)
@@ -54,7 +56,7 @@ public:
 	 * @param Buffer the chunk of data to copy
 	 * @param BufferSize the amount of data to copy
 	 */
-	void UpdatePreviewImageData(int32 Offset, const uint8* Buffer, int32 Size);
+	UE_API void UpdatePreviewImageData(int32 Offset, const uint8* Buffer, int32 Size);
 
 	/**
 	 * Copies the buffer into the AR world data (client)
@@ -63,13 +65,15 @@ public:
 	 * @param Buffer the chunk of data to copy
 	 * @param BufferSize the amount of data to copy
 	 */
-	void UpdateARWorldData(int32 Offset, const uint8* Buffer, int32 Size);
+	UE_API void UpdateARWorldData(int32 Offset, const uint8* Buffer, int32 Size);
 
 	UFUNCTION(BlueprintImplementableEvent, Category="AR Shared World", meta = (DisplayName = "OnARWorldMapIsReady", ScriptName = "OnARWorldMapIsReady"))
-	void K2_OnARWorldMapIsReady();
+	UE_API void K2_OnARWorldMapIsReady();
 
 private:
-	void TriggerCompletionIfDone();
+	UE_API void TriggerCompletionIfDone();
 	
 	bool bFiredCompletionEvent;
 };
+
+#undef UE_API
diff --git a/Engine/Source/Runtime/AugmentedReality/Public/ARSharedWorldPlayerController.h b/Engine/Source/Runtime/AugmentedReality/Public/ARSharedWorldPlayerController.h
--- a/Engine/Source/Runtime/AugmentedReality/Public/ARSharedWorldPlayerController.h
+++ b/Engine/Source/Runtime/AugmentedReality/Public/ARSharedWorldPlayerController.h
@@ -5,10 +5,12 @@
 #include "GameFramework/PlayerController.h"
 #include "ARSharedWorldPlayerController.generated.h"
 
+#define UE_API AUGMENTEDREALITY_API
+
 class AARSharedWorldGameState;
 
-UCLASS()
-class AUGMENTEDREALITY_API AARSharedWorldPlayerController :
+UCLASS(MinimalAPI)
+class AARSharedWorldPlayerController :
 	public APlayerController
 {
 	GENERATED_UCLASS_BODY()
@@ -18,7 +20,7 @@ public:
 	 * Tells the server it is ready for receiving any shared world data
 	 */
 	UFUNCTION(Reliable, Server, WithValidation)
-	void ServerMarkReadyForReceiving();
+	UE_API void ServerMarkReadyForReceiving();
 	
 	/**
 	 * Used to setup the initial values and size the arrays (client)
@@ -27,7 +29,7 @@ public:
 	 * @param ARWorldDataSize the total size in bytes of the AR world data that will be sent
 	 */
 	UFUNCTION(Reliable, Client, WithValidation)
-	void ClientInitSharedWorld(int32 PreviewImageSize, int32 ARWorldDataSize);
+	UE_API void ClientInitSharedWorld(int32 PreviewImageSize, int32 ARWorldDataSize);
 
 	/**
 	 * Copies the buffer into the image data (client)
@@ -37,7 +39,7 @@ public:
 	 * @param BufferSize the amount of data to copy
 	 */
 	UFUNCTION(Reliable, Client, WithValidation)
-	void ClientUpdatePreviewImageData(int32 Offset, const TArray<uint8>& Buffer);
+	UE_API void ClientUpdatePreviewImageData(int32 Offset, const TArray<uint8>& Buffer);
 	
 	/**
 	 * Copies the buffer into the AR world data (client)
@@ -47,16 +49,18 @@ public:
 	 * @param BufferSize the amount of data to copy
 	 */
 	UFUNCTION(Reliable, Client, WithValidation)
-	void ClientUpdateARWorldData(int32 Offset, const TArray<uint8>& Buffer);
+	UE_API void ClientUpdateARWorldData(int32 Offset, const TArray<uint8>& Buffer);
 	
 	/** @return whether this player can start receiving the ar world data */
 	bool IsReadyToReceive() const { return bIsReadyToReceive; }
 
 private:
-	virtual void Tick(float DeltaSeconds) override;
+	UE_API virtual void Tick(float DeltaSeconds) override;
 
-	AARSharedWorldGameState* GetGameState();
-	bool IsGameStateReady();
+	UE_API AARSharedWorldGameState* GetGameState();
+	UE_API bool IsGameStateReady();
 
 	bool bIsReadyToReceive;
 };
+
+#undef UE_API
diff --git a/Engine/Source/Runtime/AugmentedReality/Public/ARSkyLight.h b/Engine/Source/Runtime/AugmentedReality/Public/ARSkyLight.h
--- a/Engine/Source/Runtime/AugmentedReality/Public/ARSkyLight.h
+++ b/Engine/Source/Runtime/AugmentedReality/Public/ARSkyLight.h
@@ -7,11 +7,13 @@
 
 #include "ARSkyLight.generated.h"
 
+#define UE_API AUGMENTEDREALITY_API
+
 class UAREnvironmentCaptureProbe;
 
 /** This sky light class forces a refresh of the cube map data when an AR environment probe changes */
-UCLASS()
-class AUGMENTEDREALITY_API AARSkyLight :
+UCLASS(MinimalAPI)
+class AARSkyLight :
 	public ASkyLight
 {
 	GENERATED_UCLASS_BODY()
@@ -19,10 +21,10 @@ class AUGMENTEDREALITY_API AARSkyLight :
 public:
 	/** Sets the environment capture probe that this sky light is driven by */
 	UFUNCTION(BlueprintCallable, Category="AR AugmentedReality|SkyLight")
-	void SetEnvironmentCaptureProbe(UAREnvironmentCaptureProbe* InCaptureProbe);
+	UE_API void SetEnvironmentCaptureProbe(UAREnvironmentCaptureProbe* InCaptureProbe);
 
 private:
-	virtual void Tick(float DeltaTime) override;
+	UE_API virtual void Tick(float DeltaTime) override;
 
 	UPROPERTY()
 	TObjectPtr<UAREnvironmentCaptureProbe> CaptureProbe;
@@ -30,3 +32,5 @@ private:
 	/** The timestamp from the environment probe when we last updated the cube map */
 	float LastUpdateTimestamp;
 };
+
+#undef UE_API
diff --git a/Engine/Source/Runtime/AugmentedReality/Public/ARSystem.h b/Engine/Source/Runtime/AugmentedReality/Public/ARSystem.h
--- a/Engine/Source/Runtime/AugmentedReality/Public/ARSystem.h
+++ b/Engine/Source/Runtime/AugmentedReality/Public/ARSystem.h
@@ -74,7 +74,7 @@ public:
  * \endverbatim
  *
  */
-class AUGMENTEDREALITY_API IARSystemSupport
+class IARSystemSupport
 {
 public:
 	/** Invoked after the base AR system has been initialized. */
diff --git a/Engine/Source/Runtime/AugmentedReality/Public/ARSystemSupportBase.h b/Engine/Source/Runtime/AugmentedReality/Public/ARSystemSupportBase.h
--- a/Engine/Source/Runtime/AugmentedReality/Public/ARSystemSupportBase.h
+++ b/Engine/Source/Runtime/AugmentedReality/Public/ARSystemSupportBase.h
@@ -6,12 +6,15 @@
 
 #include "ARSystem.h"
 
-class AUGMENTEDREALITY_API FARSystemSupportBase : public IARSystemSupport
+#define UE_API AUGMENTEDREALITY_API
+
+class FARSystemSupportBase : public IARSystemSupport
 {
 public:
-	virtual bool OnPinComponentToARPin(USceneComponent* ComponentToPin, UARPin* Pin) override;
+	UE_API virtual bool OnPinComponentToARPin(USceneComponent* ComponentToPin, UARPin* Pin) override;
 };
 
 
 
 
+#undef UE_API
diff --git a/Engine/Source/Runtime/AugmentedReality/Public/ARTextures.h b/Engine/Source/Runtime/AugmentedReality/Public/ARTextures.h
--- a/Engine/Source/Runtime/AugmentedReality/Public/ARTextures.h
+++ b/Engine/Source/Runtime/AugmentedReality/Public/ARTextures.h
@@ -26,8 +26,8 @@ enum class EARTextureType : uint8
  * Derived from UTexture instead of UTexture2D because UTexture2D is all about streaming and source art
  * ? probably should have been UTexture2DDynamic
  */
-UCLASS(Abstract, BlueprintType)
-class AUGMENTEDREALITY_API UARTexture : public UTexture
+UCLASS(MinimalAPI, Abstract, BlueprintType)
+class UARTexture : public UTexture
 {
 	GENERATED_UCLASS_BODY()
 
@@ -73,8 +73,8 @@ public:
 /**
  * Base class for all AR textures that represent the camera image data
  */
-UCLASS(Abstract, BlueprintType)
-class AUGMENTEDREALITY_API UARTextureCameraImage : public UARTexture
+UCLASS(MinimalAPI, Abstract, BlueprintType)
+class UARTextureCameraImage : public UARTexture
 {
 	GENERATED_UCLASS_BODY()
 };
@@ -102,8 +102,8 @@ enum class EARDepthAccuracy : uint8
 /**
  * Base class for all AR textures that represent the camera depth data
  */
-UCLASS(Abstract, BlueprintType)
-class AUGMENTEDREALITY_API UARTextureCameraDepth : public UARTexture
+UCLASS(MinimalAPI, Abstract, BlueprintType)
+class UARTextureCameraDepth : public UARTexture
 {
 	GENERATED_UCLASS_BODY()
 
@@ -124,8 +124,8 @@ public:
 /**
  * Base class for all AR textures that represent the environment for lighting and reflection
  */
-UCLASS(Abstract, BlueprintType)
-class AUGMENTEDREALITY_API UAREnvironmentCaptureProbeTexture : public UTextureCube
+UCLASS(MinimalAPI, Abstract, BlueprintType)
+class UAREnvironmentCaptureProbeTexture : public UTextureCube
 {
 	GENERATED_UCLASS_BODY()
 	
