Updated headers using UnrealCodeFixup to make sure dllstorage is on methods/staticvars instead of types. Part 5/n

diff --git a/Engine/Source/Runtime/AugmentedReality/Public/ARTraceResult.h b/Engine/Source/Runtime/AugmentedReality/Public/ARTraceResult.h
--- a/Engine/Source/Runtime/AugmentedReality/Public/ARTraceResult.h
+++ b/Engine/Source/Runtime/AugmentedReality/Public/ARTraceResult.h
@@ -5,33 +5,35 @@
 #include "ARTypes.h"
 #include "ARTraceResult.generated.h"
 
+#define UE_API AUGMENTEDREALITY_API
+
 class FARSupportInterface ;
 
 /**
  * A result of an intersection found during a hit-test.
  */
 USTRUCT( BlueprintType, Category="AR AugmentedReality", meta=(Experimental))
-struct AUGMENTEDREALITY_API FARTraceResult
+struct FARTraceResult
 {
 	GENERATED_BODY();
 	
-	FARTraceResult();
+	UE_API FARTraceResult();
 	
-	FARTraceResult( const TSharedPtr<FARSupportInterface , ESPMode::ThreadSafe>& InARSystem, float InDistanceFromCamera, EARLineTraceChannels InTraceChannel, const FTransform& InLocalToTrackingTransform, UARTrackedGeometry* InTrackedGeometry );
+	UE_API FARTraceResult( const TSharedPtr<FARSupportInterface , ESPMode::ThreadSafe>& InARSystem, float InDistanceFromCamera, EARLineTraceChannels InTraceChannel, const FTransform& InLocalToTrackingTransform, UARTrackedGeometry* InTrackedGeometry );
 	
-	float GetDistanceFromCamera() const;
+	UE_API float GetDistanceFromCamera() const;
 
-	void SetLocalToWorldTransform(const FTransform& LocalToWorldTransform);
+	UE_API void SetLocalToWorldTransform(const FTransform& LocalToWorldTransform);
 	
-	FTransform GetLocalToTrackingTransform() const;
+	UE_API FTransform GetLocalToTrackingTransform() const;
 	
-	FTransform GetLocalToWorldTransform() const;
+	UE_API FTransform GetLocalToWorldTransform() const;
 	
-	FTransform GetLocalTransform() const;
+	UE_API FTransform GetLocalTransform() const;
 	
-	UARTrackedGeometry* GetTrackedGeometry() const;
+	UE_API UARTrackedGeometry* GetTrackedGeometry() const;
 	
-	EARLineTraceChannels GetTraceChannel() const;
+	UE_API EARLineTraceChannels GetTraceChannel() const;
 
 	struct FARTraceResultComparer
 	{
@@ -73,3 +75,5 @@ class UARTraceResultDummy : public UObject
 {
 	GENERATED_BODY()
 };
+
+#undef UE_API
diff --git a/Engine/Source/Runtime/AugmentedReality/Public/ARTrackable.h b/Engine/Source/Runtime/AugmentedReality/Public/ARTrackable.h
--- a/Engine/Source/Runtime/AugmentedReality/Public/ARTrackable.h
+++ b/Engine/Source/Runtime/AugmentedReality/Public/ARTrackable.h
@@ -6,68 +6,70 @@
 #include "ARComponent.h"
 #include "ARTrackable.generated.h"
 
+#define UE_API AUGMENTEDREALITY_API
+
 class FARSupportInterface;
 class UAREnvironmentCaptureProbeTexture;
 class UMRMeshComponent;
 
-UCLASS(BlueprintType)
-class AUGMENTEDREALITY_API UARTrackedGeometry : public UObject
+UCLASS(MinimalAPI, BlueprintType)
+class UARTrackedGeometry : public UObject
 {
 	GENERATED_BODY()
 	
 public:
-	UARTrackedGeometry();
+	UE_API UARTrackedGeometry();
 	
-	void InitializeNativeResource(IARRef* InNativeResource);
+	UE_API void InitializeNativeResource(IARRef* InNativeResource);
 
-	virtual void DebugDraw( UWorld* World, const FLinearColor& OutlineColor, float OutlineThickness, float PersistForSeconds = 0.0f) const;
+	UE_API virtual void DebugDraw( UWorld* World, const FLinearColor& OutlineColor, float OutlineThickness, float PersistForSeconds = 0.0f) const;
 
-	void GetNetworkPayload(FARMeshUpdatePayload& Payload);
+	UE_API void GetNetworkPayload(FARMeshUpdatePayload& Payload);
 
-	void UpdateTrackedGeometryNoMove(const TSharedRef<FARSupportInterface, ESPMode::ThreadSafe>& InTrackingSystem, uint32 FrameNumber, double Timestamp);
-	void UpdateTrackedGeometry(const TSharedRef<FARSupportInterface , ESPMode::ThreadSafe>& InTrackingSystem, uint32 FrameNumber, double Timestamp, const FTransform& InLocalToTrackingTransform, const FTransform& InAlignmentTransform );
+	UE_API void UpdateTrackedGeometryNoMove(const TSharedRef<FARSupportInterface, ESPMode::ThreadSafe>& InTrackingSystem, uint32 FrameNumber, double Timestamp);
+	UE_API void UpdateTrackedGeometry(const TSharedRef<FARSupportInterface , ESPMode::ThreadSafe>& InTrackingSystem, uint32 FrameNumber, double Timestamp, const FTransform& InLocalToTrackingTransform, const FTransform& InAlignmentTransform );
 	
-	void UpdateTrackingState( EARTrackingState NewTrackingState );
+	UE_API void UpdateTrackingState( EARTrackingState NewTrackingState );
 	
-	void UpdateAlignmentTransform( const FTransform& NewAlignmentTransform );
+	UE_API void UpdateAlignmentTransform( const FTransform& NewAlignmentTransform );
 
-	void SetDebugName( FName InDebugName );
+	UE_API void SetDebugName( FName InDebugName );
 	
-	void SetName(const FString& InName);
+	UE_API void SetName(const FString& InName);
 
-	IARRef* GetNativeResource();
+	UE_API IARRef* GetNativeResource();
 	
 	UFUNCTION(BlueprintPure, Category="AR AugmentedReality|Tracked Geometry")
-	FTransform GetLocalToWorldTransform() const;
+	UE_API FTransform GetLocalToWorldTransform() const;
 	
 	UFUNCTION(BlueprintPure, Category="AR AugmentedReality|Tracked Geometry")
-	FTransform GetLocalToTrackingTransform() const;
+	UE_API FTransform GetLocalToTrackingTransform() const;
 	
-	FTransform GetLocalToTrackingTransform_NoAlignment() const;
+	UE_API FTransform GetLocalToTrackingTransform_NoAlignment() const;
 	
 	UFUNCTION(BlueprintPure, Category="AR AugmentedReality|Tracked Geometry")
-	EARTrackingState GetTrackingState() const;
-	void SetTrackingState(EARTrackingState NewState);
+	UE_API EARTrackingState GetTrackingState() const;
+	UE_API void SetTrackingState(EARTrackingState NewState);
 
 	UFUNCTION(BlueprintPure, Category="AR AugmentedReality|Tracked Geometry")
-	bool IsTracked() const;
+	UE_API bool IsTracked() const;
 
 	UFUNCTION(BlueprintPure, Category="AR AugmentedReality|Tracked Geometry")
-	FName GetDebugName() const;
+	UE_API FName GetDebugName() const;
 	
 	UFUNCTION(BlueprintPure, Category="AR AugmentedReality|Tracked Geometry")
-	const FString& GetName() const;
+	UE_API const FString& GetName() const;
 	
 	UFUNCTION(BlueprintPure, Category="AR AugmentedReality|Tracked Geometry")
-	int32 GetLastUpdateFrameNumber() const;
+	UE_API int32 GetLastUpdateFrameNumber() const;
 	
 	UFUNCTION(BlueprintPure, Category="AR AugmentedReality|Tracked Geometry")
-	float GetLastUpdateTimestamp() const;
+	UE_API float GetLastUpdateTimestamp() const;
 	inline void SetLastUpdateTimestamp(double InTimestamp) { LastUpdateTimestamp = InTimestamp; }
 
 	UFUNCTION(BlueprintPure, Category="AR AugmentedReality|Tracked Geometry")
-	UMRMeshComponent* GetUnderlyingMesh();
-	void SetUnderlyingMesh(UMRMeshComponent* InMRMeshComponent);
+	UE_API UMRMeshComponent* GetUnderlyingMesh();
+	UE_API void SetUnderlyingMesh(UMRMeshComponent* InMRMeshComponent);
 
 	UPROPERTY(BlueprintReadOnly, Category="AR AugmentedReality|Tracked Geometry")
 	FGuid UniqueId;
@@ -81,8 +83,8 @@ public:
 	void SetSpatialMeshUsageFlags(const EARSpatialMeshUsageFlags InFlags) { SpatialMeshUsageFlags = InFlags; }
 
 protected:
-	TSharedPtr<FARSupportInterface , ESPMode::ThreadSafe> GetARSystem() const;
-	void UpdateSessionPayload(FARSessionPayload& Payload) const;
+	UE_API TSharedPtr<FARSupportInterface , ESPMode::ThreadSafe> GetARSystem() const;
+	UE_API void UpdateSessionPayload(FARSessionPayload& Payload) const;
 	
 	UPROPERTY()
 	FTransform LocalToTrackingTransform;
@@ -126,20 +128,20 @@ private:
 	FString AnchorName;
 };
 
-UCLASS(BlueprintType)
-class AUGMENTEDREALITY_API UARPlaneGeometry : public UARTrackedGeometry
+UCLASS(MinimalAPI, BlueprintType)
+class UARPlaneGeometry : public UARTrackedGeometry
 {
 	GENERATED_BODY()
 	
 public:
 
-	void UpdateTrackedGeometry(const TSharedRef<FARSupportInterface , ESPMode::ThreadSafe>& InTrackingSystem, uint32 FrameNumber, double Timestamp, const FTransform& InLocalToTrackingTransform, const FTransform& InAlignmentTransform, const FVector InCenter, const FVector InExtent );
+	UE_API void UpdateTrackedGeometry(const TSharedRef<FARSupportInterface , ESPMode::ThreadSafe>& InTrackingSystem, uint32 FrameNumber, double Timestamp, const FTransform& InLocalToTrackingTransform, const FTransform& InAlignmentTransform, const FVector InCenter, const FVector InExtent );
 	
-	void UpdateTrackedGeometry(const TSharedRef<FARSupportInterface , ESPMode::ThreadSafe>& InTrackingSystem, uint32 FrameNumber, double Timestamp, const FTransform& InLocalToTrackingTransform, const FTransform& InAlignmentTransform, const FVector InCenter, const FVector InExtent, const TArray<FVector>& InBoundingPoly, UARPlaneGeometry* InSubsumedBy);
+	UE_API void UpdateTrackedGeometry(const TSharedRef<FARSupportInterface , ESPMode::ThreadSafe>& InTrackingSystem, uint32 FrameNumber, double Timestamp, const FTransform& InLocalToTrackingTransform, const FTransform& InAlignmentTransform, const FVector InCenter, const FVector InExtent, const TArray<FVector>& InBoundingPoly, UARPlaneGeometry* InSubsumedBy);
 	
-	virtual void DebugDraw( UWorld* World, const FLinearColor& OutlineColor, float OutlineThickness, float PersistForSeconds = 0.0f) const override;
+	UE_API virtual void DebugDraw( UWorld* World, const FLinearColor& OutlineColor, float OutlineThickness, float PersistForSeconds = 0.0f) const override;
 
-	void GetNetworkPayload(FARPlaneUpdatePayload& Payload);
+	UE_API void GetNetworkPayload(FARPlaneUpdatePayload& Payload);
 
 public:
 	UFUNCTION(BlueprintPure, Category="AR AugmentedReality|Plane Geometry")
@@ -176,30 +178,30 @@ private:
 	TObjectPtr<UARPlaneGeometry> SubsumedBy = nullptr;
 };
 
-UCLASS(BlueprintType)
-class AUGMENTEDREALITY_API UARTrackedPoint : public UARTrackedGeometry
+UCLASS(MinimalAPI, BlueprintType)
+class UARTrackedPoint : public UARTrackedGeometry
 {
 	GENERATED_BODY()
 
 public:
-	virtual void DebugDraw(UWorld* World, const FLinearColor& OutlineColor, float OutlineThickness, float PersistForSeconds = 0.0f) const override;
+	UE_API virtual void DebugDraw(UWorld* World, const FLinearColor& OutlineColor, float OutlineThickness, float PersistForSeconds = 0.0f) const override;
 
-	void GetNetworkPayload(FARPointUpdatePayload& Payload);
+	UE_API void GetNetworkPayload(FARPointUpdatePayload& Payload);
 
-	void UpdateTrackedGeometry(const TSharedRef<FARSupportInterface , ESPMode::ThreadSafe>& InTrackingSystem, uint32 FrameNumber, double Timestamp, const FTransform& InLocalToTrackingTransform, const FTransform& InAlignmentTransform);
+	UE_API void UpdateTrackedGeometry(const TSharedRef<FARSupportInterface , ESPMode::ThreadSafe>& InTrackingSystem, uint32 FrameNumber, double Timestamp, const FTransform& InLocalToTrackingTransform, const FTransform& InAlignmentTransform);
 };
 
-UCLASS(BlueprintType)
-class AUGMENTEDREALITY_API UARTrackedImage : public UARTrackedGeometry
+UCLASS(MinimalAPI, BlueprintType)
+class UARTrackedImage : public UARTrackedGeometry
 {
 	GENERATED_BODY()
 
 public:
-	virtual void DebugDraw(UWorld* World, const FLinearColor& OutlineColor, float OutlineThickness, float PersistForSeconds = 0.0f) const override;
+	UE_API virtual void DebugDraw(UWorld* World, const FLinearColor& OutlineColor, float OutlineThickness, float PersistForSeconds = 0.0f) const override;
 
-	void GetNetworkPayload(FARImageUpdatePayload& Payload);
+	UE_API void GetNetworkPayload(FARImageUpdatePayload& Payload);
 
-	void UpdateTrackedGeometry(const TSharedRef<FARSupportInterface, ESPMode::ThreadSafe>& InTrackingSystem, uint32 FrameNumber, double Timestamp, const FTransform& InLocalToTrackingTransform, const FTransform& InAlignmentTransform, FVector2D InEstimatedSize, UARCandidateImage* InDetectedImage);
+	UE_API void UpdateTrackedGeometry(const TSharedRef<FARSupportInterface, ESPMode::ThreadSafe>& InTrackingSystem, uint32 FrameNumber, double Timestamp, const FTransform& InLocalToTrackingTransform, const FTransform& InAlignmentTransform, FVector2D InEstimatedSize, UARCandidateImage* InDetectedImage);
 
 	/** @see DetectedImage */
 	UFUNCTION(BlueprintPure, Category = "AR AugmentedReality|Image Detection")
@@ -215,7 +217,7 @@ public:
 	 * ARKit will return the physical size of the ARCandidate image.
 	 */
 	UFUNCTION(BlueprintPure, Category = "AR AugmentedReality|Image Detection")
-	FVector2D GetEstimateSize();
+	UE_API FVector2D GetEstimateSize();
 
 protected:
 	/** The candidate image that was detected in the scene */
@@ -227,16 +229,16 @@ protected:
 	FVector2D EstimatedSize;
 };
 
-UCLASS(BlueprintType)
-class AUGMENTEDREALITY_API UARTrackedQRCode :
+UCLASS(MinimalAPI, BlueprintType)
+class UARTrackedQRCode :
 	public UARTrackedImage
 {
 	GENERATED_BODY()
 
 public:
-	void GetNetworkPayload(FARQRCodeUpdatePayload& Payload);
+	UE_API void GetNetworkPayload(FARQRCodeUpdatePayload& Payload);
 
-	void UpdateTrackedGeometry(const TSharedRef<FARSupportInterface, ESPMode::ThreadSafe>& InTrackingSystem, uint32 FrameNumber, double Timestamp, const FTransform& InLocalToTrackingTransform, const FTransform& InAlignmentTransform, FVector2D InEstimatedSize, const FString& CodeData, int32 InVersion);
+	UE_API void UpdateTrackedGeometry(const TSharedRef<FARSupportInterface, ESPMode::ThreadSafe>& InTrackingSystem, uint32 FrameNumber, double Timestamp, const FTransform& InLocalToTrackingTransform, const FTransform& InAlignmentTransform, FVector2D InEstimatedSize, const FString& CodeData, int32 InVersion);
 
 	/** The encoded information in the qr code */
 	UPROPERTY(BlueprintReadOnly, Category="QR Code")
@@ -341,24 +343,24 @@ enum class EAREye : uint8
 
 typedef TMap<EARFaceBlendShape, float> FARBlendShapeMap;
 
-UCLASS(BlueprintType)
-class AUGMENTEDREALITY_API UARFaceGeometry : public UARTrackedGeometry
+UCLASS(MinimalAPI, BlueprintType)
+class UARFaceGeometry : public UARTrackedGeometry
 {
 	GENERATED_BODY()
 	
 public:
-	void UpdateFaceGeometry(const TSharedRef<FARSupportInterface , ESPMode::ThreadSafe>& InTrackingSystem, uint32 FrameNumber, double Timestamp, const FTransform& InTransform, const FTransform& InAlignmentTransform, FARBlendShapeMap& InBlendShapes, TArray<FVector>& InVertices, const TArray<int32>& Indices, TArray<FVector2D>& InUVs, const FTransform& InLeftEyeTransform, const FTransform& InRightEyeTransform, const FVector& InLookAtTarget);
+	UE_API void UpdateFaceGeometry(const TSharedRef<FARSupportInterface , ESPMode::ThreadSafe>& InTrackingSystem, uint32 FrameNumber, double Timestamp, const FTransform& InTransform, const FTransform& InAlignmentTransform, FARBlendShapeMap& InBlendShapes, TArray<FVector>& InVertices, const TArray<int32>& Indices, TArray<FVector2D>& InUVs, const FTransform& InLeftEyeTransform, const FTransform& InRightEyeTransform, const FVector& InLookAtTarget);
 	
-	virtual void DebugDraw( UWorld* World, const FLinearColor& OutlineColor, float OutlineThickness, float PersistForSeconds = 0.0f) const override;
+	UE_API virtual void DebugDraw( UWorld* World, const FLinearColor& OutlineColor, float OutlineThickness, float PersistForSeconds = 0.0f) const override;
 	
-	void GetNetworkPayload(FARFaceUpdatePayload& Payload);
+	UE_API void GetNetworkPayload(FARFaceUpdatePayload& Payload);
 
 public:
 	UFUNCTION(BlueprintPure, Category="AR AugmentedReality|Face Geometry")
-	float GetBlendShapeValue(EARFaceBlendShape BlendShape) const;
+	UE_API float GetBlendShapeValue(EARFaceBlendShape BlendShape) const;
 	
 	UFUNCTION(BlueprintPure, Category="AR AugmentedReality|Face Geometry")
-	const TMap<EARFaceBlendShape, float> GetBlendShapes() const;
+	UE_API const TMap<EARFaceBlendShape, float> GetBlendShapes() const;
 
 	const FARBlendShapeMap& GetBlendShapesRef() const { return BlendShapes; }
 	
@@ -367,10 +369,10 @@ public:
 	const TArray<FVector2D>& GetUVs() const { return UVs; }
 	
 	UFUNCTION(BlueprintPure, Category="AR AugmentedReality|Face Geometry")
-	const FTransform& GetLocalSpaceEyeTransform(EAREye Eye) const;
+	UE_API const FTransform& GetLocalSpaceEyeTransform(EAREye Eye) const;
 
 	UFUNCTION(BlueprintPure, Category="AR AugmentedReality|Face Geometry")
-	FTransform GetWorldSpaceEyeTransform(EAREye Eye) const;
+	UE_API FTransform GetWorldSpaceEyeTransform(EAREye Eye) const;
 
 	/** The target the eyes are looking at */
 	UPROPERTY(BlueprintReadOnly, Category="AR AugmentedReality|Face Geometry")
@@ -401,28 +403,28 @@ private:
 };
 
 /** A tracked environment texture probe that gives you a cube map for reflections */
-UCLASS(BlueprintType)
-class AUGMENTEDREALITY_API UAREnvironmentCaptureProbe :
+UCLASS(MinimalAPI, BlueprintType)
+class UAREnvironmentCaptureProbe :
 	public UARTrackedGeometry
 {
 	GENERATED_BODY()
 	
 public:
-	UAREnvironmentCaptureProbe();
+	UE_API UAREnvironmentCaptureProbe();
 	
 	/** Draw a box visulizing the bounds of the probe */
-	virtual void DebugDraw(UWorld* World, const FLinearColor& OutlineColor, float OutlineThickness, float PersistForSeconds = 0.0f) const override;
+	UE_API virtual void DebugDraw(UWorld* World, const FLinearColor& OutlineColor, float OutlineThickness, float PersistForSeconds = 0.0f) const override;
 	
-	void GetNetworkPayload(FAREnvironmentProbeUpdatePayload& Payload);
+	UE_API void GetNetworkPayload(FAREnvironmentProbeUpdatePayload& Payload);
 
-	void UpdateEnvironmentCapture(const TSharedRef<FARSupportInterface , ESPMode::ThreadSafe>& InTrackingSystem, uint32 FrameNumber, double Timestamp, const FTransform& InLocalToTrackingTransform, const FTransform& InAlignmentTransform, FVector InExtent);
+	UE_API void UpdateEnvironmentCapture(const TSharedRef<FARSupportInterface , ESPMode::ThreadSafe>& InTrackingSystem, uint32 FrameNumber, double Timestamp, const FTransform& InLocalToTrackingTransform, const FTransform& InAlignmentTransform, FVector InExtent);
 
 	/** @see Extent */
 	UFUNCTION(BlueprintPure, Category="AR AugmentedReality|Environment Capture Probe")
-	FVector GetExtent() const;
+	UE_API FVector GetExtent() const;
 	/** @see EnvironmentCaptureTexture */
 	UFUNCTION(BlueprintPure, Category="AR AugmentedReality|Environment Capture Probe")
-	UAREnvironmentCaptureProbeTexture* GetEnvironmentCaptureTexture();
+	UE_API UAREnvironmentCaptureProbeTexture* GetEnvironmentCaptureTexture();
 
 protected:
 	/** The size of area this probe covers */
@@ -434,17 +436,17 @@ protected:
 	TObjectPtr<UAREnvironmentCaptureProbeTexture> EnvironmentCaptureTexture;
 };
 
-UCLASS(BlueprintType)
-class AUGMENTEDREALITY_API UARTrackedObject : public UARTrackedGeometry
+UCLASS(MinimalAPI, BlueprintType)
+class UARTrackedObject : public UARTrackedGeometry
 {
 	GENERATED_BODY()
 	
 public:
-	virtual void DebugDraw(UWorld* World, const FLinearColor& OutlineColor, float OutlineThickness, float PersistForSeconds = 0.0f) const override;
+	UE_API virtual void DebugDraw(UWorld* World, const FLinearColor& OutlineColor, float OutlineThickness, float PersistForSeconds = 0.0f) const override;
 	
-	void GetNetworkPayload(FARObjectUpdatePayload& Payload);
+	UE_API void GetNetworkPayload(FARObjectUpdatePayload& Payload);
 
-	void UpdateTrackedGeometry(const TSharedRef<FARSupportInterface , ESPMode::ThreadSafe>& InTrackingSystem, uint32 FrameNumber, double Timestamp, const FTransform& InLocalToTrackingTransform, const FTransform& InAlignmentTransform, UARCandidateObject* InDetectedObject);
+	UE_API void UpdateTrackedGeometry(const TSharedRef<FARSupportInterface , ESPMode::ThreadSafe>& InTrackingSystem, uint32 FrameNumber, double Timestamp, const FTransform& InLocalToTrackingTransform, const FTransform& InAlignmentTransform, UARCandidateObject* InDetectedObject);
 	
 	/** @see DetectedObject */
 	UFUNCTION(BlueprintPure, Category = "AR AugmentedReality|Object Detection")
@@ -456,17 +458,17 @@ private:
 	TObjectPtr<UARCandidateObject> DetectedObject;
 };
 
-UCLASS(BlueprintType)
-class AUGMENTEDREALITY_API UARTrackedPose : public UARTrackedGeometry
+UCLASS(MinimalAPI, BlueprintType)
+class UARTrackedPose : public UARTrackedGeometry
 {
 	GENERATED_BODY()
 	
 public:
-	virtual void DebugDraw(UWorld* World, const FLinearColor& OutlineColor, float OutlineThickness, float PersistForSeconds = 0.0f) const override;
+	UE_API virtual void DebugDraw(UWorld* World, const FLinearColor& OutlineColor, float OutlineThickness, float PersistForSeconds = 0.0f) const override;
 	
-	void GetNetworkPayload(FARPoseUpdatePayload& Payload);
+	UE_API void GetNetworkPayload(FARPoseUpdatePayload& Payload);
 
-	void UpdateTrackedPose(const TSharedRef<FARSupportInterface , ESPMode::ThreadSafe>& InTrackingSystem, uint32 FrameNumber, double Timestamp, const FTransform& InLocalToTrackingTransform, const FTransform& InAlignmentTransform, const FARPose3D& InTrackedPose);
+	UE_API void UpdateTrackedPose(const TSharedRef<FARSupportInterface , ESPMode::ThreadSafe>& InTrackingSystem, uint32 FrameNumber, double Timestamp, const FTransform& InLocalToTrackingTransform, const FTransform& InAlignmentTransform, const FARPose3D& InTrackedPose);
 	
 	UFUNCTION(BlueprintPure, Category = "AR AugmentedReality|Pose Tracking")
 	const FARPose3D& GetTrackedPoseData() const { return TrackedPose; };
@@ -477,14 +479,14 @@ private:
 	FARPose3D TrackedPose;
 };
 
-UCLASS(BlueprintType)
-class AUGMENTEDREALITY_API UARMeshGeometry : public UARTrackedGeometry
+UCLASS(MinimalAPI, BlueprintType)
+class UARMeshGeometry : public UARTrackedGeometry
 {
 	GENERATED_BODY()
 	
 public:
 
-	void GetNetworkPayload(FARMeshUpdatePayload& Payload);
+	UE_API void GetNetworkPayload(FARMeshUpdatePayload& Payload);
 
 	/**
 	 * Try to determine the classification of the object at a world space location
@@ -499,16 +501,16 @@ public:
 };
 
 
-UCLASS(BlueprintType)
-class AUGMENTEDREALITY_API UARGeoAnchor : public UARTrackedGeometry
+UCLASS(MinimalAPI, BlueprintType)
+class UARGeoAnchor : public UARTrackedGeometry
 {
 	GENERATED_BODY()
 	
 public:
-	void UpdateGeoAnchor(const TSharedRef<FARSupportInterface, ESPMode::ThreadSafe>& InTrackingSystem, uint32 FrameNumber, double Timestamp, const FTransform& InLocalToTrackingTransform, const FTransform& InAlignmentTransform,
+	UE_API void UpdateGeoAnchor(const TSharedRef<FARSupportInterface, ESPMode::ThreadSafe>& InTrackingSystem, uint32 FrameNumber, double Timestamp, const FTransform& InLocalToTrackingTransform, const FTransform& InAlignmentTransform,
 						 float InLongitude, float InLatitude, float InAltitudeMeters, EARAltitudeSource InAltitudeSource);
 	
-	void GetNetworkPayload(FARGeoAnchorUpdatePayload& Payload);
+	UE_API void GetNetworkPayload(FARGeoAnchorUpdatePayload& Payload);
 	
 	UFUNCTION(BlueprintPure, Category = "AR AugmentedReality|Geo Tracking")
 	float GetLongitude() const { return Longitude; }
@@ -528,3 +530,5 @@ private:
 	float AltitudeMeters = 0.f;
 	EARAltitudeSource AltitudeSource = EARAltitudeSource::Unknown;
 };
+
+#undef UE_API
diff --git a/Engine/Source/Runtime/AugmentedReality/Public/ARTrackableNotifyComponent.h b/Engine/Source/Runtime/AugmentedReality/Public/ARTrackableNotifyComponent.h
--- a/Engine/Source/Runtime/AugmentedReality/Public/ARTrackableNotifyComponent.h
+++ b/Engine/Source/Runtime/AugmentedReality/Public/ARTrackableNotifyComponent.h
@@ -7,6 +7,8 @@
 #include "Components/ActorComponent.h"
 #include "ARTrackableNotifyComponent.generated.h"
 
+#define UE_API AUGMENTEDREALITY_API
+
 DECLARE_DYNAMIC_MULTICAST_DELEGATE_OneParam(FTrackableDelegate, UARTrackedGeometry*, TrackedGeometry);
 DECLARE_DYNAMIC_MULTICAST_DELEGATE_OneParam(FTrackablePlaneDelegate, UARPlaneGeometry*, TrackedPlane);
 DECLARE_DYNAMIC_MULTICAST_DELEGATE_OneParam(FTrackablePointDelegate, UARTrackedPoint*, TrackedPoint);
@@ -16,8 +18,8 @@ DECLARE_DYNAMIC_MULTICAST_DELEGATE_OneParam(FTrackableEnvProbeDelegate, UAREnvir
 DECLARE_DYNAMIC_MULTICAST_DELEGATE_OneParam(FTrackableObjectDelegate, UARTrackedObject*, TrackedObject);
 
 /** Component used to listen to ar trackable object changes */
-UCLASS(meta=(BlueprintSpawnableComponent))
-class AUGMENTEDREALITY_API UARTrackableNotifyComponent :
+UCLASS(MinimalAPI, meta=(BlueprintSpawnableComponent))
+class UARTrackableNotifyComponent :
 	public UActorComponent
 {
 	GENERATED_BODY()
@@ -115,8 +117,8 @@ public:
 	FTrackableObjectDelegate OnRemoveTrackedObject;
 
 private:
-	virtual void OnRegister() override;
-	virtual void OnUnregister() override;
+	UE_API virtual void OnRegister() override;
+	UE_API virtual void OnUnregister() override;
 
 	/**
 	 * Triggers the delegate type if the object is the desired class
@@ -131,17 +133,19 @@ private:
 	 *
 	 * @param Added the newly added ar trackable
 	 */
-	void OnTrackableAdded(UARTrackedGeometry* Added);
+	UE_API void OnTrackableAdded(UARTrackedGeometry* Added);
 	/**
 	 * Routes the notification to the proper delegate set
 	 *
 	 * @param Updated the updated ar trackable
 	 */
-	void OnTrackableUpdated(UARTrackedGeometry* Updated);
+	UE_API void OnTrackableUpdated(UARTrackedGeometry* Updated);
 	/**
 	 * Routes the notification to the proper delegate set
 	 *
 	 * @param Removed the removed ar trackable
 	 */
-	void OnTrackableRemoved(UARTrackedGeometry* Removed);
+	UE_API void OnTrackableRemoved(UARTrackedGeometry* Removed);
 };
+
+#undef UE_API
diff --git a/Engine/Source/Runtime/AugmentedReality/Public/ARTypes.h b/Engine/Source/Runtime/AugmentedReality/Public/ARTypes.h
--- a/Engine/Source/Runtime/AugmentedReality/Public/ARTypes.h
+++ b/Engine/Source/Runtime/AugmentedReality/Public/ARTypes.h
@@ -6,6 +6,8 @@
 #include "HAL/ThreadSafeBool.h"
 #include "ARTypes.generated.h"
 
+#define UE_API AUGMENTEDREALITY_API
+
 class USceneComponent;
 class UARPin;
 class UARTrackedGeometry;
@@ -250,7 +252,7 @@ enum class EARAltitudeSource : uint8
 
 /** The current state of the AR subsystem including an optional explanation string. */
 USTRUCT(BlueprintType)
-struct AUGMENTEDREALITY_API FARSessionStatus
+struct FARSessionStatus
 {
 public:
 	
@@ -305,8 +307,8 @@ enum class EARCandidateImageOrientation : uint8
 };
 
 /** An asset that points to an image to be detected in a scene and provides the size of the object in real life */
-UCLASS(BlueprintType)
-class AUGMENTEDREALITY_API UARCandidateImage :
+UCLASS(MinimalAPI, BlueprintType)
+class UARCandidateImage :
 	public UDataAsset
 {
 	GENERATED_BODY()
@@ -348,7 +350,7 @@ public:
 protected:
 #if WITH_EDITOR
 	/** Used to enforce physical sizes matching the aspect ratio of the images */
-	virtual void PostEditChangeProperty(FPropertyChangedEvent& PropertyChangedEvent) override;
+	UE_API virtual void PostEditChangeProperty(FPropertyChangedEvent& PropertyChangedEvent) override;
 #endif
 	
 	/** The image to detect in scenes */
@@ -373,8 +375,8 @@ protected:
 };
 
 /** An asset that points to an object to be detected in a scene */
-UCLASS(BlueprintType)
-class AUGMENTEDREALITY_API UARCandidateObject :
+UCLASS(MinimalAPI, BlueprintType)
+class UARCandidateObject :
 	public UDataAsset
 {
 	GENERATED_BODY()
@@ -417,17 +419,17 @@ private:
 /**
  * Base class for async AR requests
  */
-class AUGMENTEDREALITY_API FARAsyncTask
+class FARAsyncTask
 {
 public:
 	virtual ~FARAsyncTask() {}
 	
 	/** @return whether the task succeeded or not */
-	bool HadError() const;
+	UE_API bool HadError() const;
 	/** @return information about the error if there was one */
-	FString GetErrorString() const;
+	UE_API FString GetErrorString() const;
 	/** @return whether the task has completed or not */
-	bool IsDone() const;
+	UE_API bool IsDone() const;
 	
 protected:
 	FThreadSafeBool bIsDone;
@@ -436,19 +438,19 @@ protected:
 };
 
 /** Async task that saves the world data into a buffer */
-class AUGMENTEDREALITY_API FARSaveWorldAsyncTask :
+class FARSaveWorldAsyncTask :
 	public FARAsyncTask
 {
 public:
 	/** @return the byte array that the world was saved into. Note uses MoveTemp() for efficiency so only valid once */
-	TArray<uint8> GetSavedWorldData();
+	UE_API TArray<uint8> GetSavedWorldData();
 	
 protected:
 	TArray<uint8> WorldData;
 };
 
 /** Async task that builds a candidate object used for detection from the ar session */
-class AUGMENTEDREALITY_API FARGetCandidateObjectAsyncTask :
+class FARGetCandidateObjectAsyncTask :
 	public FARAsyncTask
 {
 public:
@@ -483,7 +485,7 @@ public:
 
 /** A specific AR video format */
 USTRUCT(BlueprintType)
-struct AUGMENTEDREALITY_API FARVideoFormat
+struct FARVideoFormat
 {
 	GENERATED_BODY()
 
@@ -531,7 +533,7 @@ public:
 
 /** Represents a hierarchy of a human pose skeleton tracked by the AR system */
 USTRUCT(BlueprintType)
-struct AUGMENTEDREALITY_API FARSkeletonDefinition
+struct FARSkeletonDefinition
 {
 	GENERATED_BODY()
 
@@ -551,7 +553,7 @@ public:
 
 /** Represents a human pose tracked in the 2D space */
 USTRUCT(BlueprintType)
-struct AUGMENTEDREALITY_API FARPose2D
+struct FARPose2D
 {
 	GENERATED_BODY()
 
@@ -571,7 +573,7 @@ public:
 
 /** Represents a human pose tracked in the 3D space */
 USTRUCT(BlueprintType)
-struct AUGMENTEDREALITY_API FARPose3D
+struct FARPose3D
 {
 	GENERATED_BODY()
 
@@ -595,7 +597,7 @@ public:
 
 /** AR camera intrinsics */
 USTRUCT(BlueprintType)
-struct AUGMENTEDREALITY_API FARCameraIntrinsics
+struct FARCameraIntrinsics
 {
 	GENERATED_BODY()
 
@@ -613,3 +615,4 @@ public:
 	FVector2D PrincipalPoint = FVector2D::ZeroVector;
 };
 
+#undef UE_API
diff --git a/Engine/Source/Runtime/BlueprintRuntime/Public/BlueprintRuntimeSettings.h b/Engine/Source/Runtime/BlueprintRuntime/Public/BlueprintRuntimeSettings.h
--- a/Engine/Source/Runtime/BlueprintRuntime/Public/BlueprintRuntimeSettings.h
+++ b/Engine/Source/Runtime/BlueprintRuntime/Public/BlueprintRuntimeSettings.h
@@ -41,8 +41,8 @@ struct FBlueprintWarningSettings
 /**
 * Implements the settings for the BlueprintRuntime module
 */
-UCLASS(config = Engine, defaultconfig)
-class BLUEPRINTRUNTIME_API UBlueprintRuntimeSettings : public UObject
+UCLASS(MinimalAPI, config = Engine, defaultconfig)
+class UBlueprintRuntimeSettings : public UObject
 {
 	GENERATED_UCLASS_BODY()
 
diff --git a/Engine/Source/Runtime/Datasmith/DatasmithCore/Private/DatasmithSceneElementsImpl.h b/Engine/Source/Runtime/Datasmith/DatasmithCore/Private/DatasmithSceneElementsImpl.h
--- a/Engine/Source/Runtime/Datasmith/DatasmithCore/Private/DatasmithSceneElementsImpl.h
+++ b/Engine/Source/Runtime/Datasmith/DatasmithCore/Private/DatasmithSceneElementsImpl.h
@@ -1590,16 +1590,16 @@ private:
 };
 
 PRAGMA_DISABLE_DEPRECATION_WARNINGS  // For Clothes - UE_DEPRECATED(5.5, "The experimental Cloth importer is no longer supported.")
-class DATASMITHCORE_API FDatasmithSceneImpl : public FDatasmithElementImpl< IDatasmithScene >
+class FDatasmithSceneImpl : public FDatasmithElementImpl< IDatasmithScene >
 PRAGMA_ENABLE_DEPRECATION_WARNINGS
 {
 public:
-	explicit FDatasmithSceneImpl(const TCHAR* InName);
+	UE_API explicit FDatasmithSceneImpl(const TCHAR* InName);
 
-	virtual void Reset() override;
+	UE_API virtual void Reset() override;
 
-	virtual const TCHAR* GetHost() const override;
-	virtual void SetHost(const TCHAR* InHostname) override;
+	UE_API virtual const TCHAR* GetHost() const override;
+	UE_API virtual void SetHost(const TCHAR* InHostname) override;
 
 	virtual const TCHAR* GetExporterVersion() const override { return *(FString&)ExporterVersion; }
 	virtual void SetExporterVersion(const TCHAR* InVersion) override { ExporterVersion = InVersion; }
@@ -1635,50 +1635,50 @@ public:
 
 	virtual void AddMesh(const TSharedPtr< IDatasmithMeshElement >& InMesh) override { Meshes.Add(InMesh); }
 	virtual int32 GetMeshesCount() const override { return Meshes.Num(); }
-	virtual TSharedPtr< IDatasmithMeshElement > GetMesh(int32 InIndex) override;
-	virtual const TSharedPtr< IDatasmithMeshElement >& GetMesh(int32 InIndex) const override;
+	UE_API virtual TSharedPtr< IDatasmithMeshElement > GetMesh(int32 InIndex) override;
+	UE_API virtual const TSharedPtr< IDatasmithMeshElement >& GetMesh(int32 InIndex) const override;
 	virtual void RemoveMesh(const TSharedPtr< IDatasmithMeshElement >& InMesh) override { Meshes.Remove(InMesh); }
-	virtual void RemoveMeshAt(int32 InIndex) override;
+	UE_API virtual void RemoveMeshAt(int32 InIndex) override;
 	virtual void EmptyMeshes() override { Meshes.Empty(); }
 
 PRAGMA_DISABLE_DEPRECATION_WARNINGS
 	UE_DEPRECATED(5.5, "The experimental Cloth importer is no longer supported.")
-	virtual void AddCloth(const TSharedPtr< IDatasmithClothElement >& InElement) override;
+	UE_API virtual void AddCloth(const TSharedPtr< IDatasmithClothElement >& InElement) override;
 	UE_DEPRECATED(5.5, "The experimental Cloth importer is no longer supported.")
-	virtual int32 GetClothesCount() const override;
+	UE_API virtual int32 GetClothesCount() const override;
 	UE_DEPRECATED(5.5, "The experimental Cloth importer is no longer supported.")
-	virtual TSharedPtr< IDatasmithClothElement > GetCloth(int32 InIndex) override;
+	UE_API virtual TSharedPtr< IDatasmithClothElement > GetCloth(int32 InIndex) override;
 	UE_DEPRECATED(5.5, "The experimental Cloth importer is no longer supported.")
-	virtual const TSharedPtr< IDatasmithClothElement >& GetCloth(int32 InIndex) const override;
+	UE_API virtual const TSharedPtr< IDatasmithClothElement >& GetCloth(int32 InIndex) const override;
 	UE_DEPRECATED(5.5, "The experimental Cloth importer is no longer supported.")
-	virtual void RemoveCloth(const TSharedPtr< IDatasmithClothElement >& InElement) override;
+	UE_API virtual void RemoveCloth(const TSharedPtr< IDatasmithClothElement >& InElement) override;
 	UE_DEPRECATED(5.5, "The experimental Cloth importer is no longer supported.")
-	virtual void RemoveClothAt(int32 InIndex) override;
+	UE_API virtual void RemoveClothAt(int32 InIndex) override;
 	UE_DEPRECATED(5.5, "The experimental Cloth importer is no longer supported.")
-	virtual void EmptyClothes() override;
+	UE_API virtual void EmptyClothes() override;
 PRAGMA_ENABLE_DEPRECATION_WARNINGS
 
 	virtual void AddActor(const TSharedPtr< IDatasmithActorElement >& InActor) override { Actors.Add(InActor);  }
 	virtual int32 GetActorsCount() const override { return Actors.Num(); }
-	virtual TSharedPtr< IDatasmithActorElement > GetActor(int32 InIndex) override;
-	virtual const TSharedPtr< IDatasmithActorElement >& GetActor(int32 InIndex) const override;
-	virtual void RemoveActor(const TSharedPtr< IDatasmithActorElement >& InActor, EDatasmithActorRemovalRule RemoveRule) override;
-	virtual void RemoveActorAt(int32 InIndex, EDatasmithActorRemovalRule RemoveRule) override;
+	UE_API virtual TSharedPtr< IDatasmithActorElement > GetActor(int32 InIndex) override;
+	UE_API virtual const TSharedPtr< IDatasmithActorElement >& GetActor(int32 InIndex) const override;
+	UE_API virtual void RemoveActor(const TSharedPtr< IDatasmithActorElement >& InActor, EDatasmithActorRemovalRule RemoveRule) override;
+	UE_API virtual void RemoveActorAt(int32 InIndex, EDatasmithActorRemovalRule RemoveRule) override;
 
 	virtual void AddMaterial(const TSharedPtr< IDatasmithBaseMaterialElement >& InMaterial) override { Materials.Add(InMaterial); }
 	virtual int32 GetMaterialsCount() const override { return Materials.Num(); }
-	virtual TSharedPtr< IDatasmithBaseMaterialElement > GetMaterial(int32 InIndex) override;
-	virtual const TSharedPtr< IDatasmithBaseMaterialElement >& GetMaterial(int32 InIndex) const override;
+	UE_API virtual TSharedPtr< IDatasmithBaseMaterialElement > GetMaterial(int32 InIndex) override;
+	UE_API virtual const TSharedPtr< IDatasmithBaseMaterialElement >& GetMaterial(int32 InIndex) const override;
 	virtual void RemoveMaterial(const TSharedPtr< IDatasmithBaseMaterialElement >& InMaterial) override { Materials.Remove(InMaterial); }
-	virtual void RemoveMaterialAt(int32 InIndex) override;
+	UE_API virtual void RemoveMaterialAt(int32 InIndex) override;
 	virtual void EmptyMaterials() override { Materials.Empty(); }
 
 	virtual void AddTexture(const TSharedPtr< IDatasmithTextureElement >& InTexture) override { Textures.Add(InTexture); }
 	virtual int32 GetTexturesCount() const override { return Textures.Num(); }
-	virtual TSharedPtr< IDatasmithTextureElement > GetTexture(int32 InIndex) override;
-	virtual const TSharedPtr< IDatasmithTextureElement >& GetTexture(int32 InIndex) const override;
+	UE_API virtual TSharedPtr< IDatasmithTextureElement > GetTexture(int32 InIndex) override;
+	UE_API virtual const TSharedPtr< IDatasmithTextureElement >& GetTexture(int32 InIndex) const override;
 	virtual void RemoveTexture(const TSharedPtr< IDatasmithTextureElement >& InTexture) override { Textures.Remove(InTexture); }
-	virtual void RemoveTextureAt(int32 InIndex) override;
+	UE_API virtual void RemoveTextureAt(int32 InIndex) override;
 	virtual void EmptyTextures() override { Textures.Empty(); }
 
 	virtual void SetPostProcess(const TSharedPtr< IDatasmithPostProcessElement >& InPostProcess) override { PostProcess.Inner = InPostProcess; }
@@ -1690,29 +1690,29 @@ PRAGMA_ENABLE_DEPRECATION_WARNINGS
 
 	virtual void AddMetaData(const TSharedPtr< IDatasmithMetaDataElement >& InMetaData) override { MetaData.Add(InMetaData); GetElementToMetaDataCache().Add(InMetaData->GetAssociatedElement(), InMetaData); }
 	virtual int32 GetMetaDataCount() const override { return MetaData.Num(); }
-	virtual TSharedPtr< IDatasmithMetaDataElement > GetMetaData(int32 InIndex) override;
-	virtual const TSharedPtr< IDatasmithMetaDataElement >& GetMetaData(int32 InIndex) const override;
-	virtual TSharedPtr< IDatasmithMetaDataElement > GetMetaData(const TSharedPtr<IDatasmithElement>& Element) override;
-	virtual const TSharedPtr< IDatasmithMetaDataElement >& GetMetaData(const TSharedPtr<IDatasmithElement>& Element) const override;
-	virtual void RemoveMetaData( const TSharedPtr<IDatasmithMetaDataElement>& Element ) override;
-	virtual void RemoveMetaDataAt(int32 InIndex) override;
+	UE_API virtual TSharedPtr< IDatasmithMetaDataElement > GetMetaData(int32 InIndex) override;
+	UE_API virtual const TSharedPtr< IDatasmithMetaDataElement >& GetMetaData(int32 InIndex) const override;
+	UE_API virtual TSharedPtr< IDatasmithMetaDataElement > GetMetaData(const TSharedPtr<IDatasmithElement>& Element) override;
+	UE_API virtual const TSharedPtr< IDatasmithMetaDataElement >& GetMetaData(const TSharedPtr<IDatasmithElement>& Element) const override;
+	UE_API virtual void RemoveMetaData( const TSharedPtr<IDatasmithMetaDataElement>& Element ) override;
+	UE_API virtual void RemoveMetaDataAt(int32 InIndex) override;
 
 	virtual void AddLevelSequence(const TSharedRef< IDatasmithLevelSequenceElement >& InSequence) override { LevelSequences.Add(InSequence);  }
 	virtual int32 GetLevelSequencesCount() const override { return LevelSequences.Num(); }
-	virtual TSharedPtr< IDatasmithLevelSequenceElement > GetLevelSequence(int32 InIndex) override;
-	virtual const TSharedPtr< IDatasmithLevelSequenceElement >& GetLevelSequence(int32 InIndex) const override;
+	UE_API virtual TSharedPtr< IDatasmithLevelSequenceElement > GetLevelSequence(int32 InIndex) override;
+	UE_API virtual const TSharedPtr< IDatasmithLevelSequenceElement >& GetLevelSequence(int32 InIndex) const override;
 	virtual void RemoveLevelSequence(const TSharedRef< IDatasmithLevelSequenceElement >& InSequence) override { LevelSequences.Remove(InSequence); }
-	virtual void RemoveLevelSequenceAt(int32 InIndex) override;
+	UE_API virtual void RemoveLevelSequenceAt(int32 InIndex) override;
 
 	virtual void AddLevelVariantSets(const TSharedPtr< IDatasmithLevelVariantSetsElement >& InLevelVariantSets) override { LevelVariantSets.Add(InLevelVariantSets);  }
 	virtual int32 GetLevelVariantSetsCount() const override { return LevelVariantSets.Num(); }
-	virtual TSharedPtr< IDatasmithLevelVariantSetsElement > GetLevelVariantSets(int32 InIndex) override;
-	virtual const TSharedPtr< IDatasmithLevelVariantSetsElement >& GetLevelVariantSets(int32 InIndex) const override;
+	UE_API virtual TSharedPtr< IDatasmithLevelVariantSetsElement > GetLevelVariantSets(int32 InIndex) override;
+	UE_API virtual const TSharedPtr< IDatasmithLevelVariantSetsElement >& GetLevelVariantSets(int32 InIndex) const override;
 	virtual void RemoveLevelVariantSets(const TSharedPtr< IDatasmithLevelVariantSetsElement >& InLevelVariantSets) override { LevelVariantSets.Remove(InLevelVariantSets); }
-	virtual void RemoveLevelVariantSetsAt(int32 InIndex) override;
+	UE_API virtual void RemoveLevelVariantSetsAt(int32 InIndex) override;
 
-	virtual void AttachActor(const TSharedPtr< IDatasmithActorElement >& NewParent, const TSharedPtr< IDatasmithActorElement >& Child, EDatasmithActorAttachmentRule AttachmentRule) override;
-	virtual void AttachActorToSceneRoot(const TSharedPtr< IDatasmithActorElement >& Child, EDatasmithActorAttachmentRule AttachmentRule) override;
+	UE_API virtual void AttachActor(const TSharedPtr< IDatasmithActorElement >& NewParent, const TSharedPtr< IDatasmithActorElement >& Child, EDatasmithActorAttachmentRule AttachmentRule) override;
+	UE_API virtual void AttachActorToSceneRoot(const TSharedPtr< IDatasmithActorElement >& Child, EDatasmithActorAttachmentRule AttachmentRule) override;
 
 private:
 
@@ -1746,7 +1746,7 @@ PRAGMA_ENABLE_DEPRECATION_WARNINGS
 
 	// Internal cache for faster metadata access per-element, should be accessed via GetMetaDataCache(), do not use directly.
 	mutable TMap< TSharedPtr< IDatasmithElement >, TSharedPtr< IDatasmithMetaDataElement> > ElementToMetaDataMap;
-	TMap< TSharedPtr< IDatasmithElement >, TSharedPtr< IDatasmithMetaDataElement> >& GetElementToMetaDataCache() const;
+	UE_API TMap< TSharedPtr< IDatasmithElement >, TSharedPtr< IDatasmithMetaDataElement> >& GetElementToMetaDataCache() const;
 };
 
 #undef UE_API
diff --git a/Engine/Source/Runtime/Experimental/ChaosVDData/Public/ChaosVDDataModule.h b/Engine/Source/Runtime/Experimental/ChaosVDData/Public/ChaosVDDataModule.h
--- a/Engine/Source/Runtime/Experimental/ChaosVDData/Public/ChaosVDDataModule.h
+++ b/Engine/Source/Runtime/Experimental/ChaosVDData/Public/ChaosVDDataModule.h
@@ -4,6 +4,6 @@
 
 #include "Modules/ModuleInterface.h"
 
-class CHAOSVDDATA_API FChaosVDDataModule : public IModuleInterface
+class FChaosVDDataModule : public IModuleInterface
 {
 };
diff --git a/Engine/Source/Runtime/Experimental/ChaosVDData/Public/ChaosVDRecording.h b/Engine/Source/Runtime/Experimental/ChaosVDData/Public/ChaosVDRecording.h
--- a/Engine/Source/Runtime/Experimental/ChaosVDData/Public/ChaosVDRecording.h
+++ b/Engine/Source/Runtime/Experimental/ChaosVDData/Public/ChaosVDRecording.h
@@ -22,6 +22,8 @@
 
 #include "ChaosVDRecording.generated.h"
 
+#define UE_API CHAOSVDDATA_API
+
 namespace Chaos::VisualDebugger
 {
 	class FChaosVDSerializableNameTable;
@@ -32,7 +34,7 @@ DECLARE_MULTICAST_DELEGATE_TwoParams(FChaosVDGeometryDataLoaded, const Chaos::FC
 /** Handle to user-defined data types in a CVD recorded frame.
  * These handles share ownership of the data
  */
-struct CHAOSVDDATA_API FChaosVDCustomUserDataHandle
+struct FChaosVDCustomUserDataHandle
 {
 	/** Creates a handle for the provided data instance. This handle will hold a reference to the data
 	 * @param DataPtr Shared ref pointer to the data we want to rpovide access to
@@ -41,7 +43,7 @@ struct CHAOSVDDATA_API FChaosVDCustomUserDataHandle
 	static FChaosVDCustomUserDataHandle MakeHandle(const TSharedRef<DataStructType>& DataPtr);
 
 	/** Returns the FName of the UStruct type that this handle represents */
-	FName GetTypeName() const;
+	UE_API FName GetTypeName() const;
 
 	/** Returns a raw ptr to the data this handle provides access to.*/
 	template <typename DataStructType>
@@ -104,9 +106,9 @@ bool FChaosVDCustomUserDataHandle::IsA_Internal(const TSharedPtr<FStructOnScope>
 	return false;
 }
 
-struct CHAOSVDDATA_API FChaosVDCustomFrameData
+struct FChaosVDCustomFrameData
 {
-	void AddData(const FChaosVDCustomUserDataHandle& InData);
+	UE_API void AddData(const FChaosVDCustomUserDataHandle& InData);
 
 	/** Returns the shared ptr to a custom data instance, if such data type was added to this frame.
 	 * If the data type was not loaded, the ptr will be null
@@ -205,7 +207,7 @@ enum class EChaosVDSolverFrameAttributes : uint16
 };
 ENUM_CLASS_FLAGS(EChaosVDSolverFrameAttributes)
 
-struct CHAOSVDDATA_API FChaosVDSolverFrameData
+struct FChaosVDSolverFrameData
 {
 	PRAGMA_DISABLE_DEPRECATION_WARNINGS
 	FChaosVDSolverFrameData() = default;
@@ -402,9 +404,9 @@ ENUM_CLASS_FLAGS(EChaosVDRecordingAttributes)
  * Struct that represents a recorded Physics simulation.
  * It is currently populated while analyzing a Trace session
  */
-struct CHAOSVDDATA_API FChaosVDRecording
+struct FChaosVDRecording
 {
-	FChaosVDRecording();
+	UE_API FChaosVDRecording();
 
 	/* Constant used to define inline allocators -
 	* Unless there are some scenarios with a lot of RBAN solvers in the recording, we usually don't go over 3 tracks most of the time so 16 should be plenty by default */
@@ -417,8 +419,8 @@ struct CHAOSVDDATA_API FChaosVDRecording
 	}
 	
 	/** Returns the current available Game Frames */
-	int32 GetAvailableGameFramesNumber() const;
-	int32 GetAvailableGameFramesNumber_AssumesLocked() const;
+	UE_API int32 GetAvailableGameFramesNumber() const;
+	UE_API int32 GetAvailableGameFramesNumber_AssumesLocked() const;
 
 	/** Returns a reference to the array holding all the available game frames */
 	const TArray<FChaosVDGameFrameData>& GetAvailableGameFrames_AssumesLocked() const
@@ -436,14 +438,14 @@ struct CHAOSVDDATA_API FChaosVDRecording
 	 * Returns the number of available frame data for the specified solver ID
 	 * @param SolverID ID of the solver 
 	 */
-	int32 GetAvailableSolverFramesNumber(int32 SolverID) const;
-	int32 GetAvailableSolverFramesNumber_AssumesLocked(int32 SolverID) const;
+	UE_API int32 GetAvailableSolverFramesNumber(int32 SolverID) const;
+	UE_API int32 GetAvailableSolverFramesNumber_AssumesLocked(int32 SolverID) const;
 	
 	/**
 	 * Returns the name of the specified solver id
 	 * @param SolverID ID of the solver 
 	 */
-	FName GetSolverFName(int32 SolverID);
+	UE_API FName GetSolverFName(int32 SolverID);
 
 	UE_DEPRECATED(5.6, "Please use the GetSolverFName instead")
 	FString GetSolverName(int32 SolverID)
@@ -455,10 +457,10 @@ struct CHAOSVDDATA_API FChaosVDRecording
 	 * Returns the name of the specified solver id. Must be called from within a ReadLock
 	 * @param SolverID ID of the solver
 	 */
-	FName GetSolverFName_AssumedLocked(int32 SolverID);
+	UE_API FName GetSolverFName_AssumedLocked(int32 SolverID);
 
-	bool IsServerSolver_AssumesLocked(int32 SolverID);
-	bool IsServerSolver(int32 SolverID);
+	UE_API bool IsServerSolver_AssumesLocked(int32 SolverID);
+	UE_API bool IsServerSolver(int32 SolverID);
 
 	UE_DEPRECATED(5.6, "Please use the GetSolverFName_AssumedLocked instead")
 	FString GetSolverName_AssumedLocked(int32 SolverID)
@@ -473,7 +475,7 @@ struct CHAOSVDDATA_API FChaosVDRecording
 	 * @param bKeyFrameOnly True if we should return a keyframe (real or generated) for the provided frame number if available or nothing
 	 * @return Ptr to the existing solver frame data from the specified ID and Frame number - It is a ptr to the array element, Do not store
 	 */
-	FChaosVDSolverFrameData* GetSolverFrameData_AssumesLocked(int32 SolverID, int32 FrameNumber, bool bKeyFrameOnly = false);
+	UE_API FChaosVDSolverFrameData* GetSolverFrameData_AssumesLocked(int32 SolverID, int32 FrameNumber, bool bKeyFrameOnly = false);
 	
 	/**
 	 * Return a ptr to the existing solver frame data from the specified ID and Frame number
@@ -481,7 +483,7 @@ struct CHAOSVDDATA_API FChaosVDRecording
 	 * @param Cycle Platform cycle at which the solver frame was recorded
 	 * @return Ptr to the existing solver frame data from the specified ID and Frame number - It is a ptr to the array element, Do not store
 	 */
-	FChaosVDSolverFrameData* GetSolverFrameDataAtCycle_AssumesLocked(int32 SolverID, uint64 Cycle);
+	UE_API FChaosVDSolverFrameData* GetSolverFrameDataAtCycle_AssumesLocked(int32 SolverID, uint64 Cycle);
 
 	/**
 	 * Searches and returns the lowest frame number of a solver at the specified cycle
@@ -489,12 +491,12 @@ struct CHAOSVDDATA_API FChaosVDRecording
 	 * @param Cycle Platform cycle to use as lower bound
 	 * @return Found frame number. INDEX_NONE if no frame is found for the specified cycle
 	 */
-	int32 GetLowestSolverFrameNumberAtCycle(int32 SolverID, uint64 Cycle);
-	int32 GetLowestSolverFrameNumberAtCycle_AssumesLocked(int32 SolverID, uint64 Cycle);
+	UE_API int32 GetLowestSolverFrameNumberAtCycle(int32 SolverID, uint64 Cycle);
+	UE_API int32 GetLowestSolverFrameNumberAtCycle_AssumesLocked(int32 SolverID, uint64 Cycle);
 
-	int32 GetLowestSolverFrameNumberAtNetworkFrameNumber_AssumesLocked(int32 SolverID, int32 NetworkFrameNumber);
+	UE_API int32 GetLowestSolverFrameNumberAtNetworkFrameNumber_AssumesLocked(int32 SolverID, int32 NetworkFrameNumber);
 
-	int32 FindFirstSolverKeyFrameNumberFromFrame_AssumesLocked(int32 SolverID, int32 StartFrameNumber);
+	UE_API int32 FindFirstSolverKeyFrameNumberFromFrame_AssumesLocked(int32 SolverID, int32 StartFrameNumber);
 	
 	/**
 	 * Searches and returns the lowest frame number of a solver at the specified cycle
@@ -502,8 +504,8 @@ struct CHAOSVDDATA_API FChaosVDRecording
 	 * @param GameFrame Platform cycle to use as lower bound
 	 * @return Found frame number. INDEX_NONE if no frame is found for the specified cycle
 	 */
-	int32 GetLowestSolverFrameNumberGameFrame(int32 SolverID, int32 GameFrame);
-	int32 GetLowestSolverFrameNumberGameFrame_AssumesLocked(int32 SolverID, int32 GameFrame);
+	UE_API int32 GetLowestSolverFrameNumberGameFrame(int32 SolverID, int32 GameFrame);
+	UE_API int32 GetLowestSolverFrameNumberGameFrame_AssumesLocked(int32 SolverID, int32 GameFrame);
 	
 	/**
 	 * Searches and returns the lowest game frame number at the specified solver frame
@@ -511,21 +513,21 @@ struct CHAOSVDDATA_API FChaosVDRecording
 	 * @param SolverFrame Frame number of the solver to evaluate
 	 * @return Found Game frame number. INDEX_NONE if no frame is found for the specified cycle
 	 */
-	int32 GetLowestGameFrameAtSolverFrameNumber(int32 SolverID, int32 SolverFrame);
-	int32 GetLowestGameFrameAtSolverFrameNumber_AssumesLocked(int32 SolverID, int32 SolverFrame);
+	UE_API int32 GetLowestGameFrameAtSolverFrameNumber(int32 SolverID, int32 SolverFrame);
+	UE_API int32 GetLowestGameFrameAtSolverFrameNumber_AssumesLocked(int32 SolverID, int32 SolverFrame);
 
 	/**
 	 * Adds a Solver Frame Data entry for a specific Solver ID. Creates a solver entry if it does not exist 
 	 * @param SolverID ID of the solver to add
 	 * @param InFrameData Reference to the frame data we want to add
 	 */
-	void AddFrameForSolver(const int32 SolverID, FChaosVDSolverFrameData&& InFrameData);
+	UE_API void AddFrameForSolver(const int32 SolverID, FChaosVDSolverFrameData&& InFrameData);
 
 	/**
 	 * Adds a Game Frame Data entry. Creates a solver entry if it does not exist 
 	 * @param InFrameData Reference to the frame data we want to add
 	 */
-	void AddGameFrameData(const FChaosVDGameFrameData& InFrameData);
+	UE_API void AddGameFrameData(const FChaosVDGameFrameData& InFrameData);
 
 	/** Called each time new geometry data becomes available in the recording - Mainly when a new frame is added from the Trace analysis */
 	FChaosVDGeometryDataLoaded& OnGeometryDataLoaded()
@@ -538,32 +540,32 @@ struct CHAOSVDDATA_API FChaosVDRecording
 	 * @param Cycle Platform Cycle to be used in the search
 	 * @return A ptr to the recorded game frame data - This is a ptr to the array element. Do not store
 	 */
-	FChaosVDGameFrameData* GetGameFrameDataAtCycle_AssumesLocked(uint64 Cycle);
+	UE_API FChaosVDGameFrameData* GetGameFrameDataAtCycle_AssumesLocked(uint64 Cycle);
 
 	/**
 	 * Searches for a recorded Game frame at the specified cycle 
 	 * @param FrameNumber Frame Number
 	 * @return A ptr to the recorded game frame data - This is a ptr to the array element. Do not store
 	 */
-	FChaosVDGameFrameData* GetGameFrameData_AssumesLocked(int32 FrameNumber);
+	UE_API FChaosVDGameFrameData* GetGameFrameData_AssumesLocked(int32 FrameNumber);
 
 	/** Returns a ptr to the last recorded game frame - This is a ptr to the array element. Do not store */
-	FChaosVDGameFrameData* GetLastGameFrameData_AssumesLocked();
+	UE_API FChaosVDGameFrameData* GetLastGameFrameData_AssumesLocked();
 
 	/**
 	 * Searches and returns the lowest game frame number at the specified cycle
 	 * @param Cycle Platform Cycle to be used in the search as lower bound
 	 * @return Found Game frame number. INDEX_NONE if no frame is found for the specified cycle
 	 */
-	int32 GetLowestGameFrameNumberAtCycle(uint64 Cycle);
-	int32 GetLowestGameFrameNumberAtCycle_AssumesLocked(uint64 Cycle);
+	UE_API int32 GetLowestGameFrameNumberAtCycle(uint64 Cycle);
+	UE_API int32 GetLowestGameFrameNumberAtCycle_AssumesLocked(uint64 Cycle);
 
 	/**
 	 * Searches and returns the lowest game frame number at the specified cycle
 	 * @param Time Platform Time to be used in the search as lower bound
 	 * @return Found Game frame number. INDEX_NONE if no frame is found for the specified cycle
 	 */
-	int32 GetLowestGameFrameNumberAtTime(double Time);
+	UE_API int32 GetLowestGameFrameNumberAtTime(double Time);
 
 	/**
      * Gathers all available solvers IDs at the given Game frame number
@@ -580,7 +582,7 @@ struct CHAOSVDDATA_API FChaosVDRecording
 	void GetAvailableSolverIDsAtGameFrame_AssumesLocked(const FChaosVDGameFrameData& GameFrameData, TArray<int32, TAllocator>& OutSolversID);
 
 	/** Collapses the most important frame data from a range of solver frames into a single solver frame data */
-	void CollapseSolverFramesRange_AssumesLocked(int32 SolverID, int32 StartFrame, int32 EndFrame, FChaosVDSolverFrameData& OutCollapsedFrameData);
+	UE_API void CollapseSolverFramesRange_AssumesLocked(int32 SolverID, int32 StartFrame, int32 EndFrame, FChaosVDSolverFrameData& OutCollapsedFrameData);
 
 	/** Returns a reference to the GeometryID-ImplicitObject map of this recording */
 	const TMap<uint32, Chaos::FConstImplicitObjectPtr>& GetGeometryMap() const
@@ -597,10 +599,10 @@ struct CHAOSVDDATA_API FChaosVDRecording
 	}
 
 	/** Adds a shared Implicit Object to the recording */
-	void AddImplicitObject(const uint32 ID, const Chaos::FImplicitObjectPtr& InImplicitObject);
+	UE_API void AddImplicitObject(const uint32 ID, const Chaos::FImplicitObjectPtr& InImplicitObject);
 	
 	UE_DEPRECATED(5.4, "Please use AddImplicitObject with FImplicitObjectPtr instead")
-	void AddImplicitObject(const uint32 ID, const TSharedPtr<Chaos::FImplicitObject>& InImplicitObject);
+	UE_API void AddImplicitObject(const uint32 ID, const TSharedPtr<Chaos::FImplicitObject>& InImplicitObject);
 
 	/** Session name of the trace session used to re-build this recording */
 	FString SessionName;
@@ -653,7 +655,7 @@ struct CHAOSVDDATA_API FChaosVDRecording
 	}
 
 	/** Returns true if this recording does not have any usable data */
-	bool IsEmpty() const;
+	UE_API bool IsEmpty() const;
 
 	/** Returns the last Platform Cycle on which this recording was updated (A new frame was added) */
 	uint64 GetLastUpdatedTimeAsCycle()
@@ -666,16 +668,16 @@ struct CHAOSVDDATA_API FChaosVDRecording
 		return CollisionChannelsInfoContainer;
 	}
 	
-	void SetCollisionChannelsInfoContainer(const TSharedPtr<FChaosVDCollisionChannelsInfoContainer>& InCollisionChannelsInfo);
+	UE_API void SetCollisionChannelsInfoContainer(const TSharedPtr<FChaosVDCollisionChannelsInfoContainer>& InCollisionChannelsInfo);
 
-	bool HasSolverID(int32 SolverID);
-	bool HasSolverID_AssumesLocked(int32 SolverID);
+	UE_API bool HasSolverID(int32 SolverID);
+	UE_API bool HasSolverID_AssumesLocked(int32 SolverID);
 
-	void ReserveSolverID(int32 SolverID);
-	void ReserveSolverID_AssumesLocked(int32 SolverID);
+	UE_API void ReserveSolverID(int32 SolverID);
+	UE_API void ReserveSolverID_AssumesLocked(int32 SolverID);
 
-	void CommitSolverID(int32 SolverID);
-	void CommitSolverID_AssumesLocked(int32 SolverID);
+	UE_API void CommitSolverID(int32 SolverID);
+	UE_API void CommitSolverID_AssumesLocked(int32 SolverID);
 
 	int32 GetAvailableTrackIDForRemapping()
 	{
@@ -698,16 +700,16 @@ private:
 protected:
 
 	/** Adds an Implicit Object to the recording and takes ownership of it */
-	void AddImplicitObject(const uint32 ID, const Chaos::FImplicitObject* InImplicitObject);
+	UE_API void AddImplicitObject(const uint32 ID, const Chaos::FImplicitObject* InImplicitObject);
 	
-	void AddImplicitObject_Internal(const uint32 ID, const Chaos::FConstImplicitObjectPtr& InImplicitObject);
+	UE_API void AddImplicitObject_Internal(const uint32 ID, const Chaos::FConstImplicitObjectPtr& InImplicitObject);
 
 	/** Stores a frame number of a solver that is a Key Frame -
 	 * These are used when scrubbing to make sure the visualization is in sync with what was recorded
 	 */
-	void AddKeyFrameNumberForSolver(int32 SolverID, int32 FrameNumber);
-	void AddKeyFrameNumberForSolver_AssumesLocked(int32 SolverID, int32 FrameNumber);
-	void GenerateAndStoreKeyframeForSolver_AssumesLocked(int32 SolverID, int32 CurrentFrameNumber, int32 LastKeyFrameNumber);
+	UE_API void AddKeyFrameNumberForSolver(int32 SolverID, int32 FrameNumber);
+	UE_API void AddKeyFrameNumberForSolver_AssumesLocked(int32 SolverID, int32 FrameNumber);
+	UE_API void GenerateAndStoreKeyframeForSolver_AssumesLocked(int32 SolverID, int32 CurrentFrameNumber, int32 LastKeyFrameNumber);
 
 	TMap<int32, TArray<FChaosVDSolverFrameData>> RecordedFramesDataPerSolver;
 	TMap<int32, TMap<int32, FChaosVDSolverFrameData>> GeneratedKeyFrameDataPerSolver;
@@ -819,3 +821,5 @@ void FChaosVDRecording::GetAvailableSolverIDsAtGameFrame_AssumesLocked(const FCh
 		}
 	}
 }
+
+#undef UE_API
diff --git a/Engine/Source/Runtime/LiveLinkMessageBusFramework/Public/LiveLinkProviderImpl.h b/Engine/Source/Runtime/LiveLinkMessageBusFramework/Public/LiveLinkProviderImpl.h
--- a/Engine/Source/Runtime/LiveLinkMessageBusFramework/Public/LiveLinkProviderImpl.h
+++ b/Engine/Source/Runtime/LiveLinkMessageBusFramework/Public/LiveLinkProviderImpl.h
@@ -8,6 +8,8 @@
 #include "LiveLinkProvider.h"
 #include "MessageEndpoint.h"
 
+#define UE_API LIVELINKMESSAGEBUSFRAMEWORK_API
+
 PRAGMA_DISABLE_DEPRECATION_WARNINGS
 // Subject that the application has told us about
 struct FTrackedSubject
@@ -29,7 +31,7 @@ struct FTrackedSubject
 };
 PRAGMA_ENABLE_DEPRECATION_WARNINGS
 
-struct LIVELINKMESSAGEBUSFRAMEWORK_API FLiveLinkProvider : public ILiveLinkProvider
+struct FLiveLinkProvider : public ILiveLinkProvider
 {
 private:
 	const FString ProviderName;
@@ -53,54 +55,54 @@ private:
 
 private:
 	//Message bus message handlers
-	void HandlePingMessage(const FLiveLinkPingMessage& Message,
+	UE_API void HandlePingMessage(const FLiveLinkPingMessage& Message,
 						   const TSharedRef<class IMessageContext, ESPMode::ThreadSafe>& Context);
-	void HandleHeartbeat(const FLiveLinkHeartbeatMessage& Message,
+	UE_API void HandleHeartbeat(const FLiveLinkHeartbeatMessage& Message,
 						 const TSharedRef<class IMessageContext, ESPMode::ThreadSafe>& Context);
 	// End message bus message handlers
 	
-	FTrackedSubject& GetTrackedSubject(const FName& SubjectName);
+	UE_API FTrackedSubject& GetTrackedSubject(const FName& SubjectName);
 
-	void SendSubject(FName SubjectName, const struct FTrackedSubject& Subject);
+	UE_API void SendSubject(FName SubjectName, const struct FTrackedSubject& Subject);
 
-	void SendSubjectFrame(FName SubjectName, const struct FTrackedSubject& Subject);
+	UE_API void SendSubjectFrame(FName SubjectName, const struct FTrackedSubject& Subject);
 
 	// Get the cached data for the named subject
-	FTrackedStaticData* GetLastSubjectStaticData(const FName& SubjectName);
+	UE_API FTrackedStaticData* GetLastSubjectStaticData(const FName& SubjectName);
 
-	FTrackedFrameData* GetLastSubjectFrameData(const FName& SubjectName);
+	UE_API FTrackedFrameData* GetLastSubjectFrameData(const FName& SubjectName);
 
-	void SetLastSubjectStaticData(FName SubjectName, TSubclassOf<ULiveLinkRole> Role, FLiveLinkStaticDataStruct&& StaticData, TMap<FName, FString>&& Annotations);
+	UE_API void SetLastSubjectStaticData(FName SubjectName, TSubclassOf<ULiveLinkRole> Role, FLiveLinkStaticDataStruct&& StaticData, TMap<FName, FString>&& Annotations);
 
-	void SetLastSubjectFrameData(FName SubjectName, FLiveLinkFrameDataStruct&& FrameData, TMap<FName, FString>&& SubjectAnnotations);
+	UE_API void SetLastSubjectFrameData(FName SubjectName, FLiveLinkFrameDataStruct&& FrameData, TMap<FName, FString>&& SubjectAnnotations);
 
 	// Clear a existing track subject
-	void ClearTrackedSubject(const FName& SubjectName);
+	UE_API void ClearTrackedSubject(const FName& SubjectName);
 
 	// Get the connected addresses that should receive livelink data.
-	void GetFilteredAddresses(FName SubjectName, TArray<FMessageAddress>& Addresses);
+	UE_API void GetFilteredAddresses(FName SubjectName, TArray<FMessageAddress>& Addresses);
 
 	// Serialize animation frame data into a compressed payload to reduce throughput.
-	struct FLiveLinkSerializedFrameData SerializeAnimationData(struct FLiveLinkAnimationFrameData* AnimData);
+	UE_API struct FLiveLinkSerializedFrameData SerializeAnimationData(struct FLiveLinkAnimationFrameData* AnimData);
 
 protected:
 	// Update connected addresses and send information to the connected source
-	void HandleConnectMessage(const FLiveLinkConnectMessage& Message, const TSharedRef<class IMessageContext, ESPMode::ThreadSafe>& Context);
+	UE_API void HandleConnectMessage(const FLiveLinkConnectMessage& Message, const TSharedRef<class IMessageContext, ESPMode::ThreadSafe>& Context);
 
 	// Create the message bus message endoing responsble for dispatching message bus messages to their respective handlers
-	void CreateMessageEndpoint(struct FMessageEndpointBuilder& EndpointBuilder);
+	UE_API void CreateMessageEndpoint(struct FMessageEndpointBuilder& EndpointBuilder);
 	
 	// Get the addresses of all connected instances.
-	void GetConnectedAddresses(TArray<FMessageAddress>& Addresses);
+	UE_API void GetConnectedAddresses(TArray<FMessageAddress>& Addresses);
 
 	// Validate our current connections, removing those that have timed out.
-	void ValidateConnections();
+	UE_API void ValidateConnections();
 
 	// Close a connection using its address.
-	void CloseConnection(FMessageAddress Address);
+	UE_API void CloseConnection(FMessageAddress Address);
 
 	// Get the cached data struct for a subject
-	TPair<UClass*, FLiveLinkStaticDataStruct*> GetLastSubjectStaticDataStruct(FName SubjectName);
+	UE_API TPair<UClass*, FLiveLinkStaticDataStruct*> GetLastSubjectStaticDataStruct(FName SubjectName);
 
 	template<typename MessageType>
 	void SendMessage(MessageType* Message, EMessageFlags Flags = EMessageFlags::None)
@@ -191,41 +193,43 @@ protected:
 	}
 
 	// Send a clear subject message to indicate that the subject should be removed from the connected client.
-	void SendClearSubjectToConnections(FName SubjectName);
+	UE_API void SendClearSubjectToConnections(FName SubjectName);
 
 	// Get the MessageBus address of the provider.
-	FMessageAddress GetEndpointAddress() const;
+	UE_API FMessageAddress GetEndpointAddress() const;
 
 	// Constructor for derived classes that allows specifying that no endpoint should be created.
-	FLiveLinkProvider(const FString& InProviderName, bool bInCreateEndpoint);
+	UE_API FLiveLinkProvider(const FString& InProviderName, bool bInCreateEndpoint);
 
 public:
-	FLiveLinkProvider(const FString& InProviderName);
+	UE_API FLiveLinkProvider(const FString& InProviderName);
 
-	FLiveLinkProvider(const FString& InProviderName, struct FMessageEndpointBuilder&& EndpointBuilder);
+	UE_API FLiveLinkProvider(const FString& InProviderName, struct FMessageEndpointBuilder&& EndpointBuilder);
 
-	virtual ~FLiveLinkProvider() override;
+	UE_API virtual ~FLiveLinkProvider() override;
 
-	virtual void UpdateSubject(const FName& SubjectName, const TArray<FName>& BoneNames, const TArray<int32>& BoneParents);
+	UE_API virtual void UpdateSubject(const FName& SubjectName, const TArray<FName>& BoneNames, const TArray<int32>& BoneParents);
 
-	virtual bool UpdateSubjectStaticData(const FName SubjectName, TSubclassOf<ULiveLinkRole> Role, FLiveLinkStaticDataStruct&& StaticData, const TMap<FName, FString>& ExtraAnnotations = {}) override;
+	UE_API virtual bool UpdateSubjectStaticData(const FName SubjectName, TSubclassOf<ULiveLinkRole> Role, FLiveLinkStaticDataStruct&& StaticData, const TMap<FName, FString>& ExtraAnnotations = {}) override;
 
-	virtual void ClearSubject(const FName& SubjectName);
+	UE_API virtual void ClearSubject(const FName& SubjectName);
 
-	virtual void RemoveSubject(const FName SubjectName) override;
+	UE_API virtual void RemoveSubject(const FName SubjectName) override;
 
-PRAGMA_DISABLE_DEPRECATION_WARNINGS
+UE_API PRAGMA_DISABLE_DEPRECATION_WARNINGS
 	virtual void UpdateSubjectFrame(const FName& SubjectName, const TArray<FTransform>& BoneTransforms, const TArray<FLiveLinkCurveElement>& CurveData, double Time);
 
-	virtual void UpdateSubjectFrame(const FName& SubjectName, const TArray<FTransform>& BoneTransforms, const TArray<FLiveLinkCurveElement>& CurveData,
+	UE_API virtual void UpdateSubjectFrame(const FName& SubjectName, const TArray<FTransform>& BoneTransforms, const TArray<FLiveLinkCurveElement>& CurveData,
 									const FLiveLinkMetaData& MetaData, double Time);
-PRAGMA_ENABLE_DEPRECATION_WARNINGS
+UE_API PRAGMA_ENABLE_DEPRECATION_WARNINGS
 
 	virtual bool UpdateSubjectFrameData(const FName SubjectName, FLiveLinkFrameDataStruct&& FrameData, const TMap<FName, FString>& ExtraAnnotations = {}) override;
 
-	virtual bool HasConnection() const override;
+	UE_API virtual bool HasConnection() const override;
 
-	virtual FDelegateHandle RegisterConnStatusChangedHandle(const FLiveLinkProviderConnectionStatusChanged::FDelegate& ConnStatusChanged) override;
+	UE_API virtual FDelegateHandle RegisterConnStatusChangedHandle(const FLiveLinkProviderConnectionStatusChanged::FDelegate& ConnStatusChanged) override;
 
-	virtual void UnregisterConnStatusChangedHandle(FDelegateHandle Handle) override;
+	UE_API virtual void UnregisterConnStatusChangedHandle(FDelegateHandle Handle) override;
 };
+
+#undef UE_API
diff --git a/Engine/Source/Runtime/NetworkReplayStreaming/HttpNetworkReplayStreaming/Public/HttpNetworkReplayStreaming.h b/Engine/Source/Runtime/NetworkReplayStreaming/HttpNetworkReplayStreaming/Public/HttpNetworkReplayStreaming.h
--- a/Engine/Source/Runtime/NetworkReplayStreaming/HttpNetworkReplayStreaming/Public/HttpNetworkReplayStreaming.h
+++ b/Engine/Source/Runtime/NetworkReplayStreaming/HttpNetworkReplayStreaming/Public/HttpNetworkReplayStreaming.h
@@ -399,22 +399,22 @@ public:
 
 PRAGMA_ENABLE_DEPRECATION_WARNINGS
 
-class HTTPNETWORKREPLAYSTREAMING_API FHttpNetworkReplayStreamingFactory : public INetworkReplayStreamingFactory, public FTickableGameObject
+class FHttpNetworkReplayStreamingFactory : public INetworkReplayStreamingFactory, public FTickableGameObject
 {
 public:
 	/** INetworkReplayStreamingFactory */
-	virtual TSharedPtr< INetworkReplayStreamer > CreateReplayStreamer() override;
+	UE_API virtual TSharedPtr< INetworkReplayStreamer > CreateReplayStreamer() override;
 
 	/** FTickableGameObject */
-	virtual void Tick( float DeltaTime ) override;
+	UE_API virtual void Tick( float DeltaTime ) override;
 	virtual ETickableTickType GetTickableTickType() const override { return ETickableTickType::Always; }
-	virtual TStatId GetStatId() const override;
+	UE_API virtual TStatId GetStatId() const override;
 	virtual bool IsTickableWhenPaused() const override { return true; }
-	virtual void Flush() override;
+	UE_API virtual void Flush() override;
 
-	virtual void ShutdownModule() override;
+	UE_API virtual void ShutdownModule() override;
 
-	bool HasPendingHttpRequests() const;
+	UE_API bool HasPendingHttpRequests() const;
 
 	TArray< TSharedPtr< FHttpNetworkReplayStreamer > > HttpStreamers;
 };
diff --git a/Engine/Source/Runtime/Renderer/Public/RayTracingInstanceBufferUtil.h b/Engine/Source/Runtime/Renderer/Public/RayTracingInstanceBufferUtil.h
--- a/Engine/Source/Runtime/Renderer/Public/RayTracingInstanceBufferUtil.h
+++ b/Engine/Source/Runtime/Renderer/Public/RayTracingInstanceBufferUtil.h
@@ -122,11 +122,11 @@ public:
 
 	PRAGMA_DISABLE_DEPRECATION_WARNINGS
 
-	RENDERER_API uint32 GetMaxNumInstances() const { return Data.NumNativeGPUSceneInstances + Data.NumNativeCPUInstances; }
+	uint32 GetMaxNumInstances() const { return Data.NumNativeGPUSceneInstances + Data.NumNativeCPUInstances; }
 
-	RENDERER_API TConstArrayView<FRHIRayTracingGeometry*> GetReferencedGeometries() const { return Data.ReferencedGeometries; }
-	RENDERER_API TConstArrayView<uint32> GetInstanceGeometryIndices() const { return Data.InstanceGeometryIndices; }
-	RENDERER_API TConstArrayView<uint32> GetBaseInstancePrefixSum() const { return Data.BaseInstancePrefixSum; }
+	TConstArrayView<FRHIRayTracingGeometry*> GetReferencedGeometries() const { return Data.ReferencedGeometries; }
+	TConstArrayView<uint32> GetInstanceGeometryIndices() const { return Data.InstanceGeometryIndices; }
+	TConstArrayView<uint32> GetBaseInstancePrefixSum() const { return Data.BaseInstancePrefixSum; }
 
 	PRAGMA_ENABLE_DEPRECATION_WARNINGS
 
@@ -151,4 +151,4 @@ private:
 	FByteAddressBuffer AccelerationStructureAddressesBuffer;
 };
 
-#endif // RHI_RAYTRACING
\ No newline at end of file
+#endif // RHI_RAYTRACING
